Allowing Compression Info class to be selected by cooked platform

Moved functions to check for and select Compressed Audio Info class to the audio device as the audio formats are not accessible in cooked builds (and shouldn't really need to be either).

diff --git a/Engine/Source/Developer/AudioFormatADPCM/Private/AudioFormatADPCM.cpp b/Engine/Source/Developer/AudioFormatADPCM/Private/AudioFormatADPCM.cpp
--- a/Engine/Source/Developer/AudioFormatADPCM/Private/AudioFormatADPCM.cpp
+++ b/Engine/Source/Developer/AudioFormatADPCM/Private/AudioFormatADPCM.cpp
@@ -475,17 +475,6 @@ public:
 		// Recompress is only necessary during editor previews
 		return 0;
 	}
-
-	virtual bool HasCompressedAudioInfoClass() const
-	{
-		return false;
-	}
-
-	virtual class ICompressedAudioInfo* CreateCompressedAudioInfo() const
-	{
-		// There is no equivalent of FVorbisAudioInfo for ADPCM
-		return NULL;
-	}
 };
 
 
diff --git a/Engine/Source/Developer/AudioFormatOgg/Private/AudioFormatOgg.cpp b/Engine/Source/Developer/AudioFormatOgg/Private/AudioFormatOgg.cpp
--- a/Engine/Source/Developer/AudioFormatOgg/Private/AudioFormatOgg.cpp
+++ b/Engine/Source/Developer/AudioFormatOgg/Private/AudioFormatOgg.cpp
@@ -408,8 +408,8 @@ public:
 	virtual int32 Recompress(FName Format, const TArray<uint8>& SrcBuffer, FSoundQualityInfo& QualityInfo, TArray<uint8>& OutBuffer) const
 	{
 		check(Format == NAME_OGG);
-		ICompressedAudioInfo*	AudioInfo = CreateCompressedAudioInfo();
-		int32					CompressedSize = -1;
+		FVorbisAudioInfo	AudioInfo;
+		int32				CompressedSize = -1;
 
 		// Cannot quality preview multichannel sounds
 		if( QualityInfo.NumChannels > 2 )
@@ -423,7 +423,7 @@ public:
 		}
 
 		// Parse the audio header for the relevant information
-		if (!AudioInfo->ReadCompressedInfo(CompressedDataStore.GetTypedData(), CompressedDataStore.Num(), &QualityInfo))
+		if (!AudioInfo.ReadCompressedInfo(CompressedDataStore.GetTypedData(), CompressedDataStore.Num(), &QualityInfo))
 		{
 			return 0;
 		}
@@ -431,20 +431,10 @@ public:
 		// Decompress all the sample data
 		OutBuffer.Empty(QualityInfo.SampleDataSize);
 		OutBuffer.AddZeroed(QualityInfo.SampleDataSize);
-		AudioInfo->ExpandFile(OutBuffer.GetTypedData(), &QualityInfo);
+		AudioInfo.ExpandFile(OutBuffer.GetTypedData(), &QualityInfo);
 
 		return CompressedDataStore.Num();
 	}
-
-	virtual bool HasCompressedAudioInfoClass() const
-	{
-		return true;
-	}
-
-	virtual class ICompressedAudioInfo* CreateCompressedAudioInfo() const
-	{
-		return new FVorbisAudioInfo();
-	}
 };
 
 
diff --git a/Engine/Source/Developer/AudioFormatOpus/Private/AudioFormatOpus.cpp b/Engine/Source/Developer/AudioFormatOpus/Private/AudioFormatOpus.cpp
--- a/Engine/Source/Developer/AudioFormatOpus/Private/AudioFormatOpus.cpp
+++ b/Engine/Source/Developer/AudioFormatOpus/Private/AudioFormatOpus.cpp
@@ -204,7 +204,7 @@ public:
 	virtual int32 Recompress(FName Format, const TArray<uint8>& SrcBuffer, FSoundQualityInfo& QualityInfo, TArray<uint8>& OutBuffer) const
 	{
 		check(Format == NAME_OPUS);
-		ICompressedAudioInfo*	AudioInfo = CreateCompressedAudioInfo();
+		FOpusAudioInfo	AudioInfo;
 
 		// Cannot quality preview multichannel sounds
 		if( QualityInfo.NumChannels > 2 )
@@ -219,7 +219,7 @@ public:
 		}
 
 		// Parse the opus header for the relevant information
-		if( !AudioInfo->ReadCompressedInfo( CompressedDataStore.GetTypedData(), CompressedDataStore.Num(), &QualityInfo ) )
+		if( !AudioInfo.ReadCompressedInfo( CompressedDataStore.GetTypedData(), CompressedDataStore.Num(), &QualityInfo ) )
 		{
 			return 0;
 		}
@@ -227,7 +227,7 @@ public:
 		// Decompress all the sample data
 		OutBuffer.Empty(QualityInfo.SampleDataSize);
 		OutBuffer.AddZeroed(QualityInfo.SampleDataSize);
-		AudioInfo->ExpandFile( OutBuffer.GetTypedData(), &QualityInfo );
+		AudioInfo.ExpandFile( OutBuffer.GetTypedData(), &QualityInfo );
 
 		return CompressedDataStore.Num();
 	}
@@ -240,16 +240,6 @@ public:
 		opus_encoder_destroy(Encoder);
 #endif
 	}
-
-	virtual bool HasCompressedAudioInfoClass() const
-	{
-		return true;
-	}
-
-	virtual class ICompressedAudioInfo* CreateCompressedAudioInfo() const
-	{
-		return new FOpusAudioInfo();
-	}
 };
 
 
diff --git a/Engine/Source/Developer/TargetPlatform/Public/Interfaces/IAudioFormat.h b/Engine/Source/Developer/TargetPlatform/Public/Interfaces/IAudioFormat.h
--- a/Engine/Source/Developer/TargetPlatform/Public/Interfaces/IAudioFormat.h
+++ b/Engine/Source/Developer/TargetPlatform/Public/Interfaces/IAudioFormat.h
@@ -116,12 +116,6 @@ public:
 	 */
 	virtual int32 Recompress( FName Format, const TArray<uint8>& SrcBuffer, FSoundQualityInfo& QualityInfo, TArray<uint8>& OutBuffer ) const = 0;
 
-	/** Whether this audio format has an associated info class to decompress it */
-	virtual bool HasCompressedAudioInfoClass() const = 0;
-
-	/** Creates a Compressed audio info class suitable for decompressing this audio format */
-	virtual class ICompressedAudioInfo* CreateCompressedAudioInfo() const = 0;
-
 public:
 
 	/**
diff --git a/Engine/Source/Runtime/Android/AndroidAudio/Private/AndroidAudioBuffer.cpp b/Engine/Source/Runtime/Android/AndroidAudio/Private/AndroidAudioBuffer.cpp
--- a/Engine/Source/Runtime/Android/AndroidAudio/Private/AndroidAudioBuffer.cpp
+++ b/Engine/Source/Runtime/Android/AndroidAudio/Private/AndroidAudioBuffer.cpp
@@ -51,7 +51,7 @@ FSLESSoundBuffer* FSLESSoundBuffer::CreateQueuedBuffer( FSLESAudioDevice* AudioD
 	// Prime the first two buffers and prepare the decompression
 	FSoundQualityInfo QualityInfo = { 0 };
 	
-	Buffer->DecompressionState = new FVorbisAudioInfo();
+	Buffer->DecompressionState = AudioDevice->CreateCompressedAudioInfo(InWave);
 
 	InWave->InitAudioResource( AudioDevice->GetRuntimeFormat() );
 	
diff --git a/Engine/Source/Runtime/Android/AndroidAudio/Private/AndroidAudioDevice.cpp b/Engine/Source/Runtime/Android/AndroidAudio/Private/AndroidAudioDevice.cpp
--- a/Engine/Source/Runtime/Android/AndroidAudio/Private/AndroidAudioDevice.cpp
+++ b/Engine/Source/Runtime/Android/AndroidAudio/Private/AndroidAudioDevice.cpp
@@ -215,3 +215,21 @@ FSoundSource* FSLESAudioDevice::CreateSoundSource()
 {
 	return new FSLESSoundSource(this);
 }
+
+bool FSLESAudioDevice::HasCompressedAudioInfoClass(USoundWave* SoundWave)
+{
+#if WITH_OGGVORBIS
+	return true;
+#else
+	return false;
+#endif
+}
+
+class ICompressedAudioInfo* FSLESAudioDevice::CreateCompressedAudioInfo(USoundWave* SoundWave)
+{
+#if WITH_OGGVORBIS
+	return new FVorbisAudioInfo();
+#else
+	return NULL;
+#endif
+}
diff --git a/Engine/Source/Runtime/Android/AndroidAudio/Public/AndroidAudioDevice.h b/Engine/Source/Runtime/Android/AndroidAudio/Public/AndroidAudioDevice.h
--- a/Engine/Source/Runtime/Android/AndroidAudio/Public/AndroidAudioDevice.h
+++ b/Engine/Source/Runtime/Android/AndroidAudio/Public/AndroidAudioDevice.h
@@ -241,6 +241,10 @@ public:
 		return NAME_OGG;
 	}
 
+	virtual bool HasCompressedAudioInfoClass(USoundWave* SoundWave) OVERRIDE;
+
+	virtual class ICompressedAudioInfo* CreateCompressedAudioInfo(USoundWave* SoundWave) OVERRIDE;
+
 	/** Starts up any platform specific hardware/APIs */
 	virtual bool InitializeHardware() OVERRIDE;
  
diff --git a/Engine/Source/Runtime/Engine/Private/AudioDecompress.cpp b/Engine/Source/Runtime/Engine/Private/AudioDecompress.cpp
--- a/Engine/Source/Runtime/Engine/Private/AudioDecompress.cpp
+++ b/Engine/Source/Runtime/Engine/Private/AudioDecompress.cpp
@@ -541,8 +541,9 @@ void FOpusAudioInfo::ExpandFile(uint8* DstBuffer, struct FSoundQualityInfo* Qual
  */
 void FAsyncAudioDecompressWorker::DoWork( void )
 {
-#if WITH_OGGVORBIS
-	ICompressedAudioInfo* AudioInfo = new FVorbisAudioInfo;
+	check(GEngine);
+	check(GEngine->GetAudioDevice());
+	ICompressedAudioInfo* AudioInfo = GEngine->GetAudioDevice()->CreateCompressedAudioInfo(Wave);
 	FSoundQualityInfo	QualityInfo = { 0 };
 
 	// Parse the audio header for the relevant information
@@ -582,9 +583,6 @@ void FAsyncAudioDecompressWorker::DoWork( void )
 	Wave->RemoveAudioResource();
 
 	delete AudioInfo;
-#else
-	check(0);
-#endif
 }
 
 // end
diff --git a/Engine/Source/Runtime/Engine/Private/AudioDevice.cpp b/Engine/Source/Runtime/Engine/Private/AudioDevice.cpp
--- a/Engine/Source/Runtime/Engine/Private/AudioDevice.cpp
+++ b/Engine/Source/Runtime/Engine/Private/AudioDevice.cpp
@@ -2269,8 +2269,7 @@ void FAudioDevice::Precache(USoundWave* SoundWave, bool bSynchronous, bool bTrac
 		// Streaming data, created programmatically.
 		SoundWave->DecompressionType = DTYPE_Procedural;
 	}
-#if WITH_OGGVORBIS
-	else
+	else if ( HasCompressedAudioInfoClass(SoundWave) )
 	{
 		const FSoundGroup& SoundGroup = GetDefault<USoundGroups>()->GetSoundGroup(SoundWave->SoundGroup);
 
@@ -2321,13 +2320,11 @@ void FAudioDevice::Precache(USoundWave* SoundWave, bool bSynchronous, bool bTrac
 			bTrackMemory = false;
 		}
 	}
-#else
 	else
 	{
 		// Preserve old behavior if there is no compressed audio info class for this audio format
 		SoundWave->DecompressionType = DTYPE_Native;
 	}
-#endif
 
 	if (bTrackMemory)
 	{
diff --git a/Engine/Source/Runtime/Engine/Public/AudioDevice.h b/Engine/Source/Runtime/Engine/Public/AudioDevice.h
--- a/Engine/Source/Runtime/Engine/Public/AudioDevice.h
+++ b/Engine/Source/Runtime/Engine/Public/AudioDevice.h
@@ -473,6 +473,12 @@ public:
 
 	virtual FName GetRuntimeFormat() PURE_VIRTUAL(FAudioDevice::GetRuntimeFormat,return NAME_None;);
 
+	/** Whether this SoundWave has an associated info class to decompress it */
+	virtual bool HasCompressedAudioInfoClass(USoundWave* SoundWave) { return false; }
+
+	/** Creates a Compressed audio info class suitable for decompressing this SoundWave */
+	virtual class ICompressedAudioInfo* CreateCompressedAudioInfo(USoundWave* SoundWave) { return NULL; }
+
 	/**
 	 * Check for errors and output a human readable string
 	 */
diff --git a/Engine/Source/Runtime/Mac/CoreAudio/Private/CoreAudioBuffer.cpp b/Engine/Source/Runtime/Mac/CoreAudio/Private/CoreAudioBuffer.cpp
--- a/Engine/Source/Runtime/Mac/CoreAudio/Private/CoreAudioBuffer.cpp
+++ b/Engine/Source/Runtime/Mac/CoreAudio/Private/CoreAudioBuffer.cpp
@@ -147,7 +147,7 @@ FCoreAudioSoundBuffer* FCoreAudioSoundBuffer::CreateQueuedBuffer( FCoreAudioDevi
 	// Prime the first two buffers and prepare the decompression
 	FSoundQualityInfo QualityInfo = { 0 };
 
-	Buffer->DecompressionState = new FVorbisAudioInfo();
+	Buffer->DecompressionState = CoreAudioDevice->CreateCompressedAudioInfo(Wave);
 	
 	Wave->InitAudioResource( CoreAudioDevice->GetRuntimeFormat() );
 	
diff --git a/Engine/Source/Runtime/Mac/CoreAudio/Private/CoreAudioDevice.cpp b/Engine/Source/Runtime/Mac/CoreAudio/Private/CoreAudioDevice.cpp
--- a/Engine/Source/Runtime/Mac/CoreAudio/Private/CoreAudioDevice.cpp
+++ b/Engine/Source/Runtime/Mac/CoreAudio/Private/CoreAudioDevice.cpp
@@ -432,3 +432,21 @@ void FCoreAudioDevice::SetFreeMatrixMixerInput( int32 Input )
 {
 	MatrixMixerInputStatus[ Input ] = false;
 }
+
+bool FCoreAudioDevice::HasCompressedAudioInfoClass(USoundWave* SoundWave)
+{
+#if WITH_OGGVORBIS
+	return true;
+#else
+	return false;
+#endif
+}
+
+class ICompressedAudioInfo* FCoreAudioDevice::CreateCompressedAudioInfo(USoundWave* SoundWave)
+{
+#if WITH_OGGVORBIS
+	return new FVorbisAudioInfo();
+#else
+	return NULL;
+#endif
+}
diff --git a/Engine/Source/Runtime/Mac/CoreAudio/Public/CoreAudioDevice.h b/Engine/Source/Runtime/Mac/CoreAudio/Public/CoreAudioDevice.h
--- a/Engine/Source/Runtime/Mac/CoreAudio/Public/CoreAudioDevice.h
+++ b/Engine/Source/Runtime/Mac/CoreAudio/Public/CoreAudioDevice.h
@@ -360,6 +360,10 @@ class FCoreAudioDevice : public FAudioDevice
 		return NAME_OGG;
 	}
 
+	virtual bool HasCompressedAudioInfoClass(USoundWave* SoundWave) OVERRIDE;
+
+	virtual class ICompressedAudioInfo* CreateCompressedAudioInfo(USoundWave* SoundWave) OVERRIDE;
+
 	FORCEINLINE uint32 CalculateMatrixElement( uint32 InputNum, uint32 OutputNum )
 	{
 		return ( ( InputNum << 16 ) | ( OutputNum & 0x0000FFFF ) );
diff --git a/Engine/Source/Runtime/Windows/XAudio2/Private/XAudio2Buffer.cpp b/Engine/Source/Runtime/Windows/XAudio2/Private/XAudio2Buffer.cpp
--- a/Engine/Source/Runtime/Windows/XAudio2/Private/XAudio2Buffer.cpp
+++ b/Engine/Source/Runtime/Windows/XAudio2/Private/XAudio2Buffer.cpp
@@ -326,7 +326,7 @@ FXAudio2SoundBuffer* FXAudio2SoundBuffer::CreateQueuedBuffer( FXAudio2Device* XA
 	// Prime the first two buffers and prepare the decompression
 	FSoundQualityInfo QualityInfo = { 0 };
 
-	Buffer->DecompressionState = new FVorbisAudioInfo();
+	Buffer->DecompressionState = XAudio2Device->CreateCompressedAudioInfo(Wave);
 
 	Wave->InitAudioResource(XAudio2Device->GetRuntimeFormat());
 
diff --git a/Engine/Source/Runtime/Windows/XAudio2/Private/XAudio2Device.cpp b/Engine/Source/Runtime/Windows/XAudio2/Private/XAudio2Device.cpp
--- a/Engine/Source/Runtime/Windows/XAudio2/Private/XAudio2Device.cpp
+++ b/Engine/Source/Runtime/Windows/XAudio2/Private/XAudio2Device.cpp
@@ -226,6 +226,24 @@ FSoundSource* FXAudio2Device::CreateSoundSource()
 	return new FXAudio2SoundSource(this);
 }
 
+bool FXAudio2Device::HasCompressedAudioInfoClass(USoundWave* SoundWave)
+{
+#if WITH_OGGVORBIS
+	return true;
+#else
+	return false;
+#endif
+}
+
+class ICompressedAudioInfo* FXAudio2Device::CreateCompressedAudioInfo(USoundWave* SoundWave)
+{
+#if WITH_OGGVORBIS
+	return new FVorbisAudioInfo();
+#else
+	return NULL;
+#endif
+}
+
 /**  
  * Check for errors and output a human readable string 
  */
diff --git a/Engine/Source/Runtime/Windows/XAudio2/Public/XAudio2Device.h b/Engine/Source/Runtime/Windows/XAudio2/Public/XAudio2Device.h
--- a/Engine/Source/Runtime/Windows/XAudio2/Public/XAudio2Device.h
+++ b/Engine/Source/Runtime/Windows/XAudio2/Public/XAudio2Device.h
@@ -102,6 +102,10 @@ class FXAudio2Device : public FAudioDevice
 #endif //WITH_OGGVORBIS
 	}
 
+	virtual bool HasCompressedAudioInfoClass(USoundWave* SoundWave) OVERRIDE;
+
+	virtual class ICompressedAudioInfo* CreateCompressedAudioInfo(USoundWave* SoundWave) OVERRIDE;
+
 	/** 
 	 * Check for errors and output a human readable string 
 	 */
