Multi channel support for Opus

Implemented CookSurround function for AudioFormatOpus, refactoring as much common code from the standard Cook into re-usable functions.
Changed OpusAudioInfo class to use a multistream decoder for all sounds, as 1 and 2 channel sounds can still be decoded by it. Just had to create an array of mappings that would take the internal Opus format and give back PCM data in the order that Unreal expects.
Also made the OpusDecoderWrapper a true wrapper so that you don't have to use any opus functions outside of it and it will not function if WITH_OPUS is not defined.
Added a virtual function to the CompressedAudioInfo interface, to check whether its uncompressed data will be in the Vorbis format. This is to make sure that code currently just checking whether there is a DecompressionState can now be sure that the Vorbis format is being used before routing sound to different speakers.
Made sure that you can't cause any problems by adding or removing a SoundWave from the Audio Streaming Manager more than once. Also rearranged where sounds are added to the manager so that it is only done in one place and it's safe when changing quality triggers recompression/splitting.
Added a global multidimensional array to store all the Vorbis channel orderings as they needed to be used in multiple places.
Added Logging and Memory tracking for the streaming process.

diff --git a/Engine/Source/Developer/AudioFormatOgg/Private/AudioFormatOgg.cpp b/Engine/Source/Developer/AudioFormatOgg/Private/AudioFormatOgg.cpp
--- a/Engine/Source/Developer/AudioFormatOgg/Private/AudioFormatOgg.cpp
+++ b/Engine/Source/Developer/AudioFormatOgg/Private/AudioFormatOgg.cpp
@@ -382,12 +382,10 @@ public:
 			case 6:
 			{
 				// the stream is 5.1 surround. channel order: front left, center, front right, rear left, rear right, LFE
-				ChannelOrder.Add(0);//SPEAKER_FrontLeft);
-				ChannelOrder.Add(2);//SPEAKER_FrontCenter);
-				ChannelOrder.Add(1);//SPEAKER_FrontRight);
-				ChannelOrder.Add(4);//SPEAKER_LeftSurround);
-				ChannelOrder.Add(5);//SPEAKER_RightSurround);
-				ChannelOrder.Add(3);//SPEAKER_LowFrequency);
+				for (int32 i = 0; i < NumChannels; i++)
+				{
+					ChannelOrder.Add(VorbisChannelInfo::Order[NumChannels - 1][i]);
+				}
 				break;
 			}
 
diff --git a/Engine/Source/Developer/AudioFormatOpus/Private/AudioFormatOpus.cpp b/Engine/Source/Developer/AudioFormatOpus/Private/AudioFormatOpus.cpp
--- a/Engine/Source/Developer/AudioFormatOpus/Private/AudioFormatOpus.cpp
+++ b/Engine/Source/Developer/AudioFormatOpus/Private/AudioFormatOpus.cpp
@@ -6,11 +6,13 @@
 #include "TargetPlatform.h"
 #include "OpusAudioInfo.h"
 #include "AudioFormatOpus.h"
+#include "VorbisAudioInfo.h"
+#include "Audio.h"
 
 // Need to define this so that resampler.h compiles - probably a way around this somehow
 #define OUTSIDE_SPEEX
 
-#include "opus.h"
+#include "opus_multistream.h"
 #include "speex_resampler.h"
 
 /** Use UE4 memory allocation or Opus */
@@ -65,42 +67,10 @@ public:
 		TArray<uint8> SrcBufferCopy;
 		if (QualityInfo.SampleRate != kOpusSampleRate)
 		{
-			// Initialize resampler to convert to desired rate for Opus
-			int32 err = 0;
-			SpeexResamplerState* resampler = speex_resampler_init(QualityInfo.NumChannels, QualityInfo.SampleRate, kOpusSampleRate, SPEEX_RESAMPLER_QUALITY_DESKTOP, &err);
-			if (err != RESAMPLER_ERR_SUCCESS)
+			if (!ResamplePCM(QualityInfo.NumChannels, SrcBuffer, QualityInfo.SampleRate, SrcBufferCopy, kOpusSampleRate))
 			{
 				return false;
 			}
-
-			// Calculate extra space required for sample rate
-			float Duration = (float)SrcBuffer.Num() / (QualityInfo.SampleRate * kSampleStride);
-			int32 SafeCopySize = (Duration + 1) * kOpusSampleRate * kSampleStride;
-			SrcBufferCopy.AddUninitialized(SafeCopySize);
-			uint32 SrcSamples = SrcBuffer.Num() / kSampleStride;
-			uint32 CopySamples = SrcBufferCopy.Num() / kSampleStride;
-
-			// Do resampling and check results
-			if (QualityInfo.NumChannels == 1)
-			{
-				err = speex_resampler_process_int(resampler, 0, (const short*)(SrcBuffer.GetTypedData()), &SrcSamples, (short*)(SrcBufferCopy.GetTypedData()), &CopySamples);
-			}
-			else
-			{
-				err = speex_resampler_process_interleaved_int(resampler, (const short*)(SrcBuffer.GetTypedData()), &SrcSamples, (short*)(SrcBufferCopy.GetTypedData()), &CopySamples);
-			}
-
-			speex_resampler_destroy(resampler);
-			if (err != RESAMPLER_ERR_SUCCESS)
-			{
-				return false;
-			}
-
-			int32 WrittenBytes = (int32)(CopySamples * kSampleStride);
-			if (WrittenBytes < SrcBufferCopy.Num())
-			{
-				SrcBufferCopy.SetNum(WrittenBytes);
-			}
 		}
 		else
 		{
@@ -124,9 +94,7 @@ public:
 			return false;
 		}
 
-		// Basic attempt at mapping quality to bit rate range similar to what vorbis would give
-		int32 BitRate = FMath::GetMappedRangeValue(FVector2D(1, 100), FVector2D(64000, 400000), QualityInfo.Quality);
-		opus_encoder_ctl(Encoder, OPUS_SET_BITRATE(BitRate));
+		opus_encoder_ctl(Encoder, OPUS_SET_BITRATE(GetBitRateFromQuality(QualityInfo.Quality)));
 
 		// Create a buffer to store compressed data
 		CompressedDataStore.Empty();
@@ -145,16 +113,9 @@ public:
 			FramesToEncode++;
 		}
 
-		// Store Identifier, True Sample Count, Number of channels and Frames to Encode first
-		const char* OpusIdentifier = OPUS_ID_STRING;
-		CompressedData.Serialize((void*)OpusIdentifier, FCStringAnsi::Strlen( OpusIdentifier )+1);
-		CompressedData.Serialize(&TrueSampleCount, sizeof(uint32));
-		check(QualityInfo.NumChannels < MAX_uint8)
-		uint8 SerializedChannels = QualityInfo.NumChannels;
-		CompressedData.Serialize(&SerializedChannels, sizeof(uint8));
-		check(FramesToEncode < MAX_uint16);
-		uint16 SerializedFrames = FramesToEncode;
-		CompressedData.Serialize(&SerializedFrames, sizeof(uint16));
+		check(QualityInfo.NumChannels <= MAX_uint8);
+		check(FramesToEncode <= MAX_uint16);
+		SerializeHeaderData(CompressedData, TrueSampleCount, QualityInfo.NumChannels, FramesToEncode);
 
 		// Temporary storage with more than enough to store any compressed frame
 		TArray<uint8> TempCompressedData;
@@ -166,8 +127,8 @@ public:
 
 			if (CompressedLength < 0)
 			{
-				//const char* ErrorStr = opus_strerror(CompressedLength);
-				//UE_LOG(LogAudio, Warning, TEXT("Failed to encode: [%d] %s"), CompressedLength, ANSI_TO_TCHAR(ErrorStr));
+				const char* ErrorStr = opus_strerror(CompressedLength);
+				UE_LOG(LogAudio, Warning, TEXT("Failed to encode: [%d] %s"), CompressedLength, ANSI_TO_TCHAR(ErrorStr));
 
 				Destroy(Encoder);
 
@@ -178,9 +139,7 @@ public:
 			{
 				// Store frame length and copy compressed data before incrementing pointers
 				check(CompressedLength < MAX_uint16);
-				uint16 SerializedFrameLength = CompressedLength;
-				CompressedData.Serialize(&SerializedFrameLength, sizeof(uint16));
-				CompressedData.Serialize(TempCompressedData.GetTypedData(), CompressedLength);
+				SerialiseFrameData(CompressedData, TempCompressedData.GetTypedData(), CompressedLength);
 				SrcBufferOffset += kBytesPerFrame;
 			}
 		}
@@ -193,14 +152,162 @@ public:
 	virtual bool CookSurround(FName Format, const TArray<TArray<uint8> >& SrcBuffers, FSoundQualityInfo& QualityInfo, TArray<uint8>& CompressedDataStore) const
 	{
 		check(Format == NAME_OPUS);
+
+		// Constant sample rate for simplicity
+		const int32 kOpusSampleRate = 48000;
+		// Frame size must be one of 2.5, 5, 10, 20, 40 or 60 ms
+		const int32 kOpusFrameSizeMs = 60;
+		// Calculate frame size required by Opus
+		const int32 kOpusFrameSizeSamples = (kOpusSampleRate * kOpusFrameSizeMs) / 1000;
+		const uint32 kSampleStride = SAMPLE_SIZE * QualityInfo.NumChannels;
+		const int32 kBytesPerFrame = kOpusFrameSizeSamples * kSampleStride;
+
+		// Check whether source has compatible sample rate
+		TArray<TArray<uint8>> SrcBufferCopies;
+		if (QualityInfo.SampleRate != kOpusSampleRate)
 		{
-			// Create a buffer to store compressed data
-			CompressedDataStore.Empty();
-			FMemoryWriter CompressedData(CompressedDataStore);
-			uint32 BufferOffset = 0;
+			for (int32 Index = 0; Index < SrcBuffers.Num(); Index++)
+			{
+				TArray<uint8>& NewCopy = *new (SrcBufferCopies) TArray<uint8>;
+				if (!ResamplePCM(1, SrcBuffers[Index], QualityInfo.SampleRate, NewCopy, kOpusSampleRate))
+				{
+					return false;
+				}
+			}
 		}
-		return CompressedDataStore.Num() > 0;
+		else
+		{
+			// Take a copy of the source regardless
+			for (int32 Index = 0; Index < SrcBuffers.Num(); Index++)
+			{
+				SrcBufferCopies[Index] = SrcBuffers[Index];
+			}
+		}
+
+		// Ensure that all channels are the same length
+		int32 SourceSize = -1;
+		for (int32 Index = 0; Index < SrcBufferCopies.Num(); Index++)
+		{
+			if (!Index)
+			{
+				SourceSize = SrcBufferCopies[Index].Num();
+			}
+			else
+			{
+				if (SourceSize != SrcBufferCopies[Index].Num())
+				{
+					return false;
+				}
+			}
+		}
+		if (SourceSize <= 0)
+		{
+			return false;
+		}
+
+		// Initialise the Opus multistream encoder
+		OpusMSEncoder* Encoder = NULL;
+		int32 EncError = 0;
+		int32 streams = 0;
+		int32 coupled_streams = 0;
+		// mapping_family not documented but figured out: 0 = 1 or 2 channels, 1 = 1 to 8 channel surround sound, 255 = up to 255 channels with no surround processing
+		int32 mapping_family = 1;
+		TArray<uint8> mapping;
+		mapping.AddUninitialized(QualityInfo.NumChannels);
+#if USE_UE4_MEM_ALLOC
+		int32 EncSize = opus_multistream_surround_encoder_get_size(QualityInfo.NumChannels, mapping_family);
+		Encoder = (OpusMSEncoder*)FMemory::Malloc(EncSize);
+		EncError = opus_multistream_surround_encoder_init(Encoder, kOpusSampleRate, QualityInfo.NumChannels, mapping_family, &streams, &coupled_streams, mapping.GetTypedData(), OPUS_APPLICATION_AUDIO);
+#else
+		Encoder = opus_multistream_surround_encoder_create(kOpusSampleRate, QualityInfo.NumChannels, mapping_family, &streams, &coupled_streams, mapping.GetTypedData(), OPUS_APPLICATION_AUDIO, &EncError);
+#endif
+		if (EncError != OPUS_OK)
+		{
+			Destroy(Encoder);
+			return false;
+		}
+
+		opus_multistream_encoder_ctl(Encoder, OPUS_SET_BITRATE(GetBitRateFromQuality(QualityInfo.Quality)));
+
+		// Create a buffer to store compressed data
+		CompressedDataStore.Empty();
+		FMemoryWriter CompressedData(CompressedDataStore);
+		int32 SrcBufferOffset = 0;
+
+		// Calc frame and sample count
+		int32 FramesToEncode = SourceSize / (kOpusFrameSizeSamples * SAMPLE_SIZE);
+		uint32 TrueSampleCount = SourceSize / SAMPLE_SIZE;
+
+		// Add another frame if Source does not divide into an equal number of frames
+		if (SourceSize % (kOpusFrameSizeSamples * SAMPLE_SIZE) != 0)
+		{
+			FramesToEncode++;
+		}
+
+		check(QualityInfo.NumChannels <= MAX_uint8);
+		check(FramesToEncode <= MAX_uint16);
+		SerializeHeaderData(CompressedData, TrueSampleCount, QualityInfo.NumChannels, FramesToEncode);
 
+		// Temporary storage for source data in an interleaved format
+		TArray<uint8> TempInterleavedSrc;
+		TempInterleavedSrc.AddUninitialized(kBytesPerFrame);
+
+		// Temporary storage with more than enough to store any compressed frame
+		TArray<uint8> TempCompressedData;
+		TempCompressedData.AddUninitialized(kBytesPerFrame);
+
+		while (SrcBufferOffset < SourceSize)
+		{
+			// Read a frames worth of data from the source and pack it into interleaved temporary storage
+			for (int32 SampleIndex = 0; SampleIndex < kOpusFrameSizeSamples; ++SampleIndex)
+			{
+				int32 CurrSrcOffset = SrcBufferOffset + SampleIndex*SAMPLE_SIZE;
+				int32 CurrInterleavedOffset = SampleIndex*kSampleStride;
+				if (CurrSrcOffset < SourceSize)
+				{
+					for (uint32 ChannelIndex = 0; ChannelIndex < QualityInfo.NumChannels; ++ChannelIndex)
+					{
+						// Interleave the channels in the Vorbis format, so that the correct channel is used for LFE
+						int32 OrderedChannelIndex = VorbisChannelInfo::Order[QualityInfo.NumChannels - 1][ChannelIndex];
+						int32 CurrInterleavedIndex = CurrInterleavedOffset + ChannelIndex*SAMPLE_SIZE;
+
+						// Copy both bytes that make up a single sample
+						TempInterleavedSrc[CurrInterleavedIndex] = SrcBufferCopies[OrderedChannelIndex][CurrSrcOffset];
+						TempInterleavedSrc[CurrInterleavedIndex + 1] = SrcBufferCopies[OrderedChannelIndex][CurrSrcOffset + 1];
+					}
+				}
+				else
+				{
+					// Zero the rest of the temp buffer to make it an exact frame
+					FMemory::Memzero(TempInterleavedSrc.GetTypedData() + CurrInterleavedOffset, kBytesPerFrame - CurrInterleavedOffset);
+					SampleIndex = kOpusFrameSizeSamples;
+				}
+			}
+
+			int32 CompressedLength = opus_multistream_encode(Encoder, (const opus_int16*)(TempInterleavedSrc.GetTypedData()), kOpusFrameSizeSamples, TempCompressedData.GetTypedData(), TempCompressedData.Num());
+
+			if (CompressedLength < 0)
+			{
+				const char* ErrorStr = opus_strerror(CompressedLength);
+				UE_LOG(LogAudio, Warning, TEXT("Failed to encode: [%d] %s"), CompressedLength, ANSI_TO_TCHAR(ErrorStr));
+
+				Destroy(Encoder);
+
+				CompressedDataStore.Empty();
+				return false;
+			}
+			else
+			{
+				// Store frame length and copy compressed data before incrementing pointers
+				check(CompressedLength < MAX_uint16);
+				SerialiseFrameData(CompressedData, TempCompressedData.GetTypedData(), CompressedLength);
+				SrcBufferOffset += kOpusFrameSizeSamples * SAMPLE_SIZE;
+			}
+		}
+
+		Destroy(Encoder);
+
+		return CompressedDataStore.Num() > 0;
 	}
 
 	virtual int32 Recompress(FName Format, const TArray<uint8>& SrcBuffer, FSoundQualityInfo& QualityInfo, TArray<uint8>& OutBuffer) const
@@ -270,15 +377,7 @@ public:
 
 			if ( (ReadOffset + sizeof(uint16) + FrameSize) - WriteOffset >= kMaxChunkSizeBytes)
 			{
-				// need to write some data
-				int32 OutIndex = OutBuffers.Num();
-				int32 BytesToCopy = ReadOffset - WriteOffset;
-				OutBuffers.Add(TArray<uint8>());
-				OutBuffers[OutIndex].Empty(BytesToCopy);
-				OutBuffers[OutIndex].AddUninitialized(BytesToCopy);
-				void* NewChunkData = OutBuffers[OutIndex].GetTypedData();
-				FMemory::Memcpy(NewChunkData, LockedSrc+WriteOffset, BytesToCopy);
-				WriteOffset += BytesToCopy;
+				WriteOffset += AddDataChunk(OutBuffers, LockedSrc + WriteOffset, ReadOffset - WriteOffset);
 			}
 
 			ReadOffset += sizeof(uint16) + FrameSize;
@@ -286,20 +385,81 @@ public:
 		}
 		if (WriteOffset < ReadOffset)
 		{
-			// need to write some data
-			int32 OutIndex = OutBuffers.Num();
-			int32 BytesToCopy = ReadOffset - WriteOffset;
-			OutBuffers.Add(TArray<uint8>());
-			OutBuffers[OutIndex].Empty(BytesToCopy);
-			OutBuffers[OutIndex].AddUninitialized(BytesToCopy);
-			void* NewChunkData = OutBuffers[OutIndex].GetTypedData();
-			FMemory::Memcpy(NewChunkData, LockedSrc + WriteOffset, BytesToCopy);
-			WriteOffset += BytesToCopy;
+			WriteOffset += AddDataChunk(OutBuffers, LockedSrc + WriteOffset, ReadOffset - WriteOffset);
 		}
 
 		return true;
 	}
 
+	bool ResamplePCM(uint32 NumChannels, const TArray<uint8>& InBuffer, uint32 InSampleRate, TArray<uint8>& OutBuffer, uint32 OutSampleRate) const
+	{
+		// Initialize resampler to convert to desired rate for Opus
+		int32 err = 0;
+		SpeexResamplerState* resampler = speex_resampler_init(NumChannels, InSampleRate, OutSampleRate, SPEEX_RESAMPLER_QUALITY_DESKTOP, &err);
+		if (err != RESAMPLER_ERR_SUCCESS)
+		{
+			speex_resampler_destroy(resampler);
+			return false;
+		}
+
+		// Calculate extra space required for sample rate
+		const uint32 SampleStride = SAMPLE_SIZE * NumChannels;
+		const float Duration = (float)InBuffer.Num() / (InSampleRate * SampleStride);
+		const int32 SafeCopySize = (Duration + 1) * OutSampleRate * SampleStride;
+		OutBuffer.Empty(SafeCopySize);
+		OutBuffer.AddUninitialized(SafeCopySize);
+		uint32 InSamples = InBuffer.Num() / SampleStride;
+		uint32 OutSamples = OutBuffer.Num() / SampleStride;
+
+		// Do resampling and check results
+		if (NumChannels == 1)
+		{
+			err = speex_resampler_process_int(resampler, 0, (const short*)(InBuffer.GetTypedData()), &InSamples, (short*)(OutBuffer.GetTypedData()), &OutSamples);
+		}
+		else
+		{
+			err = speex_resampler_process_interleaved_int(resampler, (const short*)(InBuffer.GetTypedData()), &InSamples, (short*)(OutBuffer.GetTypedData()), &OutSamples);
+		}
+
+		speex_resampler_destroy(resampler);
+		if (err != RESAMPLER_ERR_SUCCESS)
+		{
+			return false;
+		}
+
+		// reduce the size of Out Buffer if more space than necessary was allocated
+		const int32 WrittenBytes = (int32)(OutSamples * SampleStride);
+		if (WrittenBytes < OutBuffer.Num())
+		{
+			OutBuffer.SetNum(WrittenBytes, true);
+		}
+
+		return true;
+	}
+
+	int32 GetBitRateFromQuality(int32 Quality) const
+	{
+		// Basic attempt at mapping quality to a bit rate range similar to Vorbis
+		// TODO: Not sure if range should be expanded or to make sure that it creates
+		// frequently used bit rates ie. 128000, 192000 etc.
+		return FMath::GetMappedRangeValue(FVector2D(1, 100), FVector2D(64000, 400000), Quality);
+	}
+
+	void SerializeHeaderData(FMemoryWriter& CompressedData, uint32 TrueSampleCount, uint8 NumChannels, uint16 NumFrames) const
+	{
+		const char* OpusIdentifier = OPUS_ID_STRING;
+		CompressedData.Serialize((void*)OpusIdentifier, FCStringAnsi::Strlen(OpusIdentifier) + 1);
+		CompressedData.Serialize(&TrueSampleCount, sizeof(uint32));
+		CompressedData.Serialize(&NumChannels, sizeof(uint8));
+		CompressedData.Serialize(&NumFrames, sizeof(uint16));
+	}
+
+	void SerialiseFrameData(FMemoryWriter& CompressedData, uint8* FrameData, uint16 FrameSize) const
+	{
+		CompressedData.Serialize(&FrameSize, sizeof(uint16));
+		CompressedData.Serialize(FrameData, FrameSize);
+	}
+
 	void Destroy(OpusEncoder* Encoder) const
 	{
 #if USE_UE4_MEM_ALLOC
@@ -308,6 +468,32 @@ public:
 		opus_encoder_destroy(Encoder);
 #endif
 	}
+
+	void Destroy(OpusMSEncoder* Encoder) const
+	{
+#if USE_UE4_MEM_ALLOC
+		FMemory::Free(Encoder);
+#else
+		opus_multistream_encoder_destroy(Encoder);
+#endif
+	}
+
+	/**
+	 * Adds a new chunk of data to the array
+	 *
+	 * @param	OutBuffers	Array of buffers to add to
+	 * @param	ChunkData	Pointer to chunk data
+	 * @param	ChunkSize	How much data to write
+	 * @return	How many bytes were written
+	 */
+	int32 AddDataChunk(TArray<TArray<uint8>>& OutBuffers, const uint8* ChunkData, int32 ChunkSize) const
+	{
+		TArray<uint8>& NewBuffer = *new (OutBuffers) TArray<uint8>;
+		NewBuffer.Empty(ChunkSize);
+		NewBuffer.AddUninitialized(ChunkSize);
+		FMemory::Memcpy(NewBuffer.GetTypedData(), ChunkData, ChunkSize);
+		return ChunkSize;
+	}
 };
 
 
diff --git a/Engine/Source/Runtime/Engine/Private/AudioStreaming.cpp b/Engine/Source/Runtime/Engine/Private/AudioStreaming.cpp
--- a/Engine/Source/Runtime/Engine/Private/AudioStreaming.cpp
+++ b/Engine/Source/Runtime/Engine/Private/AudioStreaming.cpp
@@ -33,6 +33,8 @@ FAsyncStreamDerivedChunkWorker::FAsyncStreamDerivedChunkWorker(
 /** Retrieves the derived chunk from the derived data cache. */
 void FAsyncStreamDerivedChunkWorker::DoWork()
 {
+	UE_LOG(LogAudio, Verbose, TEXT("Start of ASync DDC Chunk read for key: %s"), *DerivedDataKey);
+
 	TArray<uint8> DerivedChunkData;
 
 	if (GetDerivedDataCacheRef().GetSynchronous(*DerivedDataKey, DerivedChunkData))
@@ -49,6 +51,8 @@ void FAsyncStreamDerivedChunkWorker::DoWork()
 	}
 	FPlatformMisc::MemoryBarrier();
 	ThreadSafeCounter->Decrement();
+
+	UE_LOG(LogAudio, Verbose, TEXT("End of ASync DDC Chunk read for key: %s"), *DerivedDataKey);
 }
 
 #endif // #if WITH_EDITORONLY_DATA
@@ -110,6 +114,17 @@ bool FStreamingWaveData::UpdateStreamingStatus()
 	{
 		if (RequestStatus == AudioState_ReadyFor_Finalization)
 		{
+			if (UE_LOG_ACTIVE(LogAudio, Log) && IndicesToLoad.Num() > 0)
+			{
+				FString LogString = FString::Printf(TEXT("Finalised loading of chunk(s) %d"), IndicesToLoad[0]);
+				for (int32 Index = 1; Index < IndicesToLoad.Num(); ++ Index)
+				{
+					LogString += FString::Printf(TEXT(", %d"), IndicesToLoad[Index]);
+				}
+				LogString += FString::Printf(TEXT(" from SoundWave'%s'"), *SoundWave->GetName());
+				UE_LOG(LogAudio, Log, TEXT("%s"), *LogString);
+			}
+
 			bool bFailedRequests = false;
 #if WITH_EDITORONLY_DATA
 			bFailedRequests = FinishDDCRequests();
@@ -167,6 +182,17 @@ bool FStreamingWaveData::HasPendingRequests(TArray<uint32>& IndicesToLoad, TArra
 
 void FStreamingWaveData::BeginPendingRequests(const TArray<uint32>& IndicesToLoad, const TArray<uint32>& IndicesToFree)
 {
+	if (UE_LOG_ACTIVE(LogAudio, Log) && IndicesToLoad.Num() > 0)
+	{
+		FString LogString = FString::Printf(TEXT("Requesting ASync load of chunk(s) %d"), IndicesToLoad[0]);
+		for (int32 Index = 1; Index < IndicesToLoad.Num(); ++Index)
+		{
+			LogString += FString::Printf(TEXT(", %d"), IndicesToLoad[Index]);
+		}
+		LogString += FString::Printf(TEXT(" from SoundWave'%s'"), *SoundWave->GetName());
+		UE_LOG(LogAudio, Log, TEXT("%s"), *LogString);
+	}
+
 	// Mark Chunks for removal in case they can be reused
 	TArray<uint32> FreeChunkIndices;
 	for (auto Index : IndicesToFree)
@@ -304,11 +330,14 @@ bool FStreamingWaveData::FinishDDCRequests()
 
 FLoadedAudioChunk* FStreamingWaveData::AddNewLoadedChunk(int32 ChunkSize)
 {
-	// TODO: Might update memory management stats here
 	int32 NewIndex = LoadedChunks.Num();
 	LoadedChunks.AddZeroed();
 	LoadedChunks[NewIndex].Data = static_cast<uint8*>(FMemory::Malloc(ChunkSize));
 	LoadedChunks[NewIndex].DataSize = LoadedChunks[NewIndex].MemorySize = ChunkSize;
+
+	INC_DWORD_STAT_BY(STAT_AudioMemorySize, ChunkSize);
+	INC_DWORD_STAT_BY(STAT_AudioMemory, ChunkSize);
+
 	return &LoadedChunks[NewIndex];
 }
 
@@ -316,8 +345,12 @@ void FStreamingWaveData::FreeLoadedChunk(FLoadedAudioChunk& LoadedChunk)
 {
 	if (LoadedChunk.Data != NULL)
 	{
-		// TODO: Might update memory management stats here
 		FMemory::Free(LoadedChunk.Data);
+
+		// Stat housekeeping
+		DEC_DWORD_STAT_BY(STAT_AudioMemorySize, LoadedChunk.MemorySize);
+		DEC_DWORD_STAT_BY(STAT_AudioMemory, LoadedChunk.MemorySize);
+
 		LoadedChunk.Data = NULL;
 		LoadedChunk.DataSize = 0;
 		LoadedChunk.MemorySize = 0;
@@ -346,7 +379,7 @@ void FAudioStreamingManager::UpdateResourceStreaming(float DeltaTime, bool bProc
 
 	for (auto Source : StreamingSoundSources)
 	{
-		USoundWave* Wave = Source->GetWaveInstance()->WaveData;
+		USoundWave* Wave = Source->GetWaveInstance() ? Source->GetWaveInstance()->WaveData : NULL;
 		FStreamingWaveData* WaveData = StreamingSoundWaves.Find(Wave);
 
 		if (WaveData && WaveData->PendingChunkChangeRequestStatus.GetValue() == AudioState_ReadyFor_Requests)
@@ -358,9 +391,10 @@ void FAudioStreamingManager::UpdateResourceStreaming(float DeltaTime, bool bProc
 			{
 				WaveRequest.RequiredIndices.AddUnique(SourceChunk);
 				WaveRequest.RequiredIndices.AddUnique((SourceChunk+1)%Wave->RunningPlatformData->NumChunks);
-				if (Source->GetBuffer()->GetCurrentChunkOffset() > Wave->RunningPlatformData->Chunks[SourceChunk].DataSize / 2)
+				if (!WaveData->LoadedChunkIndices.Contains(SourceChunk)
+				|| Source->GetBuffer()->GetCurrentChunkOffset() > Wave->RunningPlatformData->Chunks[SourceChunk].DataSize / 2)
 				{
-					// already read over half, request is high priority
+					// currently not loaded or already read over half, request is high priority
 					WaveRequest.bPrioritiseRequest = true;
 				}
 			}
@@ -409,10 +443,9 @@ void FAudioStreamingManager::RemoveLevel(class ULevel* Level)
 
 void FAudioStreamingManager::AddStreamingSoundWave(USoundWave* SoundWave)
 {
-	if (FPlatformProperties::SupportsAudioStreaming() && SoundWave->IsStreaming())
+	if (FPlatformProperties::SupportsAudioStreaming() && SoundWave->IsStreaming()
+	&& StreamingSoundWaves.Find(SoundWave) == NULL)
 	{
-		check(StreamingSoundWaves.Find(SoundWave) == NULL);
-
 		FStreamingWaveData& WaveData = StreamingSoundWaves.Add(SoundWave);
 		WaveData.Initialize(SoundWave);
 	}
diff --git a/Engine/Source/Runtime/Engine/Private/OpusAudioInfo.cpp b/Engine/Source/Runtime/Engine/Private/OpusAudioInfo.cpp
--- a/Engine/Source/Runtime/Engine/Private/OpusAudioInfo.cpp
+++ b/Engine/Source/Runtime/Engine/Private/OpusAudioInfo.cpp
@@ -9,47 +9,90 @@
 #define WITH_OPUS (PLATFORM_WINDOWS || PLATFORM_MAC || PLATFORM_LINUX)
 
 #if WITH_OPUS
-#include "opus.h"
+#include "opus_multistream.h"
 #endif
 
-/*------------------------------------------------------------------------------------
-FOpusDecoderWrapper
-------------------------------------------------------------------------------------*/
 #define USE_UE4_MEM_ALLOC 1
 #define OPUS_SAMPLE_RATE 48000
 #define OPUS_MAX_FRAME_SIZE 5760
 #define SAMPLE_SIZE			( ( uint32 )sizeof( short ) )
 
+///////////////////////////////////////////////////////////////////////////////////////
+// Followed pattern used in opus_multistream_encoder.c - this will allow us to setup //
+// a multistream decoder without having to save extra information for every asset.   //
+///////////////////////////////////////////////////////////////////////////////////////
+struct UnrealChannelLayout{
+	int32 NumStreams;
+	int32 NumCoupledStreams;
+	uint8 Mapping[8];
+};
+
+/* Index is NumChannels-1*/
+static const UnrealChannelLayout UnrealMappings[8] = {
+	{ 1, 0, { 0 } },                      /* 1: mono */
+	{ 1, 1, { 0, 1 } },                   /* 2: stereo */
+	{ 2, 1, { 0, 1, 2 } },                /* 3: 1-d surround */
+	{ 2, 2, { 0, 1, 2, 3 } },             /* 4: quadraphonic surround */
+	{ 3, 2, { 0, 1, 4, 2, 3 } },          /* 5: 5-channel surround */
+	{ 4, 2, { 0, 1, 4, 5, 2, 3 } },       /* 6: 5.1 surround */
+	{ 4, 3, { 0, 1, 4, 6, 2, 3, 5 } },    /* 7: 6.1 surround */
+	{ 5, 3, { 0, 1, 6, 7, 2, 3, 4, 5 } }, /* 8: 7.1 surround */
+};
+
+/*------------------------------------------------------------------------------------
+FOpusDecoderWrapper
+------------------------------------------------------------------------------------*/
 struct FOpusDecoderWrapper
 {
 	FOpusDecoderWrapper(uint8 NumChannels)
 	{
 #if WITH_OPUS
-#if USE_UE4_MEM_ALLOC
-		int32 DecSize = opus_decoder_get_size(NumChannels);
-		Decoder = (OpusDecoder*)FMemory::Malloc(DecSize);
-		DecError = opus_decoder_init(Decoder, OPUS_SAMPLE_RATE, NumChannels);
-#else
-		Decoder = opus_decoder_create(OPUS_SAMPLE_RATE, NumChannels, &DecError);
-#endif
+		check(NumChannels <= 8);
+		const UnrealChannelLayout& Layout = UnrealMappings[NumChannels-1];
+	#if USE_UE4_MEM_ALLOC
+		int32 DecSize = opus_multistream_decoder_get_size(Layout.NumStreams, Layout.NumCoupledStreams);
+		Decoder = (OpusMSDecoder*)FMemory::Malloc(DecSize);
+		DecError = opus_multistream_decoder_init(Decoder, OPUS_SAMPLE_RATE, NumChannels, Layout.NumStreams, Layout.NumCoupledStreams, Layout.Mapping);
+	#else
+		Decoder = opus_multistream_decoder_create(OPUS_SAMPLE_RATE, NumChannels, Layout.NumStreams, Layout.NumCoupledStreams, Layout.Mapping, &DecError);
+	#endif
 #endif
 	}
 
 	~FOpusDecoderWrapper()
 	{
 #if WITH_OPUS
-#if USE_UE4_MEM_ALLOC
+	#if USE_UE4_MEM_ALLOC
 		FMemory::Free(Decoder);
-#else
-		opus_encoder_destroy(Decoder);
+	#else
+		opus_multistream_encoder_destroy(Decoder);
+	#endif
 #endif
+	}
+
+	int32 Decode(const uint8* FrameData, uint16 FrameSize, int16* OutPCMData, int32 SampleSize)
+	{
+#if WITH_OPUS
+		return opus_multistream_decode(Decoder, FrameData, FrameSize, OutPCMData, SampleSize, 0);
+#else
+		return -1;
 #endif
 	}
 
+	bool WasInitialisedSuccessfully() const
+	{
 #if WITH_OPUS
-	OpusDecoder* Decoder;
+		return DecError == OPUS_OK;
+#else
+		return false;
 #endif
+	}
+
+private:
+#if WITH_OPUS
+	OpusMSDecoder* Decoder;
 	int32 DecError;
+#endif
 };
 
 /*------------------------------------------------------------------------------------
@@ -114,18 +157,17 @@ bool FOpusAudioInfo::ReadCompressedInfo(const uint8* InSrcBufferData, uint32 InS
 	AudioDataOffset = SrcBufferOffset;
 	SampleStride = SAMPLE_SIZE * NumChannels;
 
+	check(OpusDecoderWrapper == NULL);
 	OpusDecoderWrapper = new FOpusDecoderWrapper(NumChannels);
-	LastDecodedPCM.Empty();
-	LastDecodedPCM.AddUninitialized(OPUS_MAX_FRAME_SIZE * SampleStride);
-#if WITH_OPUS
-	if (OpusDecoderWrapper->DecError != OPUS_OK)
+	if (!OpusDecoderWrapper->WasInitialisedSuccessfully())
 	{
 		delete OpusDecoderWrapper;
+		OpusDecoderWrapper = NULL;
 		return false;
 	}
-#else
-	return false;
-#endif
+
+	LastDecodedPCM.Empty(OPUS_MAX_FRAME_SIZE * SampleStride);
+	LastDecodedPCM.AddUninitialized(OPUS_MAX_FRAME_SIZE * SampleStride);
 
 	if (QualityInfo)
 	{
@@ -252,6 +294,8 @@ bool FOpusAudioInfo::StreamCompressedData(uint8* Destination, bool bLooping, uin
 
 	SCOPE_CYCLE_COUNTER(STAT_OpusDecompressTime);
 
+	UE_LOG(LogAudio, Log, TEXT("Streaming compressed data from SoundWave'%s' - Chunk %d, Offset %d"), *StreamingSoundWave->GetName(), CurrentChunkIndex, SrcBufferOffset);
+
 	// Write out any PCM data that was decoded during the last request
 	uint32 RawPCMOffset = WriteFromDecodedPCM(Destination, BufferSize);
 
@@ -266,7 +310,8 @@ bool FOpusAudioInfo::StreamCompressedData(uint8* Destination, bool bLooping, uin
 		}
 		else
 		{
-			// Still not loaded, should probably have some kind of error, zero current buffer
+			// Still not loaded, zero remainder of current buffer
+			UE_LOG(LogAudio, Warning, TEXT("Unable to read from chunk %d of SoundWave'%s'"), CurrentChunkIndex, *StreamingSoundWave->GetName());
 			ZeroBuffer(Destination + RawPCMOffset, BufferSize - RawPCMOffset);
 			return false;
 		}
@@ -333,6 +378,7 @@ bool FOpusAudioInfo::StreamCompressedData(uint8* Destination, bool bLooping, uin
 				SrcBufferData = IStreamingManager::Get().GetAudioStreamingManager().GetLoadedChunk(StreamingSoundWave, CurrentChunkIndex);
 				if (SrcBufferData)
 				{
+					UE_LOG(LogAudio, Log, TEXT("Incremented current chunk from SoundWave'%s' - Chunk %d, Offset %d"), *StreamingSoundWave->GetName(), CurrentChunkIndex, SrcBufferOffset);
 					SrcBufferDataSize = StreamingSoundWave->RunningPlatformData->Chunks[CurrentChunkIndex].DataSize;
 				}
 				else
@@ -354,11 +400,7 @@ int32 FOpusAudioInfo::DecompressToPCMBuffer(uint16 FrameSize)
 	const uint8* SrcPtr = SrcBufferData + SrcBufferOffset;
 	SrcBufferOffset += FrameSize;
 	LastPCMOffset = 0;
-#if WITH_OPUS
-	return opus_decode(OpusDecoderWrapper->Decoder, SrcPtr, FrameSize, (int16*)(LastDecodedPCM.GetTypedData()), OPUS_MAX_FRAME_SIZE, 0);
-#else
-	return -1;
-#endif
+	return OpusDecoderWrapper->Decode(SrcPtr, FrameSize, (int16*)LastDecodedPCM.GetTypedData(), OPUS_MAX_FRAME_SIZE);
 }
 
 uint32 FOpusAudioInfo::IncrementCurrentSampleCount(uint32 NewSamples)
diff --git a/Engine/Source/Runtime/Engine/Private/SoundWave.cpp b/Engine/Source/Runtime/Engine/Private/SoundWave.cpp
--- a/Engine/Source/Runtime/Engine/Private/SoundWave.cpp
+++ b/Engine/Source/Runtime/Engine/Private/SoundWave.cpp
@@ -415,14 +415,6 @@ void USoundWave::PostEditChangeProperty(FPropertyChangedEvent& PropertyChangedEv
 			FreeResources();
 			UpdatePlatformData();
 			MarkPackageDirty();
-			if (IsStreaming())
-			{
-				IStreamingManager::Get().GetAudioStreamingManager().AddStreamingSoundWave(this);
-			}
-			else
-			{
-				IStreamingManager::Get().GetAudioStreamingManager().RemoveStreamingSoundWave(this);
-			}
 		}
 	}
 }
@@ -683,10 +675,18 @@ void USoundWave::UpdatePlatformData()
 		}
 
 #if WITH_EDITORONLY_DATA
+		// Temporarily remove from streaming manager to release currently used data chunks
+		IStreamingManager::Get().GetAudioStreamingManager().RemoveStreamingSoundWave(this);
 		// Recache platform data if the source has changed.
 		CachePlatformData();
+		// Add back to the streaming manager to reload first chunk
+		IStreamingManager::Get().GetAudioStreamingManager().AddStreamingSoundWave(this);
 #endif
 	}
+	else
+	{
+		IStreamingManager::Get().GetAudioStreamingManager().RemoveStreamingSoundWave(this);
+	}
 }
 
 void USoundWave::GetChunkData(int32 ChunkIndex, uint8** OutChunkData)
diff --git a/Engine/Source/Runtime/Engine/Private/VorbisAudioInfo.cpp b/Engine/Source/Runtime/Engine/Private/VorbisAudioInfo.cpp
--- a/Engine/Source/Runtime/Engine/Private/VorbisAudioInfo.cpp
+++ b/Engine/Source/Runtime/Engine/Private/VorbisAudioInfo.cpp
@@ -30,6 +30,24 @@
 #define VORBIS_BYTE_ORDER 1
 #endif
 
+/**
+ * Channel order expected for a multi-channel ogg vorbis file.
+ * Ordering taken from http://xiph.org/vorbis/doc/Vorbis_I_spec.html#x1-800004.3.9
+ */
+namespace VorbisChannelInfo
+{
+	const int32 Order[8][8] = {
+		{ 0 },
+		{ 0, 1 },
+		{ 0, 2, 1 },
+		{ 0, 1, 2, 3 },
+		{ 0, 2, 1, 3, 4 },
+		{ 0, 2, 1, 4, 5, 3 },
+		{ 0, 2, 1, 4, 5, 6, 3 },
+		{ 0, 2, 1, 4, 5, 6, 7, 3 }
+	};
+};
+
 /*------------------------------------------------------------------------------------
 	FVorbisFileWrapper. Hides Vorbis structs from public headers.
 ------------------------------------------------------------------------------------*/
diff --git a/Engine/Source/Runtime/Engine/Public/AudioDecompress.h b/Engine/Source/Runtime/Engine/Public/AudioDecompress.h
--- a/Engine/Source/Runtime/Engine/Public/AudioDecompress.h
+++ b/Engine/Source/Runtime/Engine/Public/AudioDecompress.h
@@ -64,6 +64,12 @@ public:
 	 */
 	virtual uint32 GetSourceBufferSize() const = 0;
 
+	/**
+	 * Whether the decompressed audio will be arranged using Vorbis' channel ordering
+	 * See http://www.xiph.org/vorbis/doc/Vorbis_I_spec.html#x1-800004.3.9 for details
+	 */
+	virtual bool UsesVorbisChannelOrdering() const = 0;
+
 	////////////////////////////////////////////////////////////////
 	// Following functions are optional if streaming is supported //
 	////////////////////////////////////////////////////////////////
diff --git a/Engine/Source/Runtime/Engine/Public/OpusAudioInfo.h b/Engine/Source/Runtime/Engine/Public/OpusAudioInfo.h
--- a/Engine/Source/Runtime/Engine/Public/OpusAudioInfo.h
+++ b/Engine/Source/Runtime/Engine/Public/OpusAudioInfo.h
@@ -28,6 +28,7 @@ public:
 	ENGINE_API virtual void ExpandFile(uint8* DstBuffer, struct FSoundQualityInfo* QualityInfo) override;
 	ENGINE_API virtual void EnableHalfRate(bool HalfRate) override {};
 	virtual uint32 GetSourceBufferSize() const override { return SrcBufferDataSize;}
+	virtual bool UsesVorbisChannelOrdering() const override { return false; }
 
 	virtual bool SupportsStreaming() const override {return true;}
 	virtual bool StreamCompressedInfo(USoundWave* Wave, struct FSoundQualityInfo* QualityInfo) override;
diff --git a/Engine/Source/Runtime/Engine/Public/VorbisAudioInfo.h b/Engine/Source/Runtime/Engine/Public/VorbisAudioInfo.h
--- a/Engine/Source/Runtime/Engine/Public/VorbisAudioInfo.h
+++ b/Engine/Source/Runtime/Engine/Public/VorbisAudioInfo.h
@@ -19,6 +19,11 @@
 	#endif
 #endif
 
+namespace VorbisChannelInfo
+{
+	extern ENGINE_API const int32 Order[8][8];
+}
+
 /**
  * Loads vorbis dlls
 */
@@ -74,6 +79,8 @@ public:
 
 	virtual uint32 GetSourceBufferSize() const { return SrcBufferDataSize;}
 
+	virtual bool UsesVorbisChannelOrdering() const override { return true; }
+
 	struct FVorbisFileWrapper* VFWrapper;
 	const uint8*		SrcBufferData;
 	uint32			SrcBufferDataSize;
diff --git a/Engine/Source/Runtime/Mac/CoreAudio/Private/CoreAudioSource.cpp b/Engine/Source/Runtime/Mac/CoreAudio/Private/CoreAudioSource.cpp
--- a/Engine/Source/Runtime/Mac/CoreAudio/Private/CoreAudioSource.cpp
+++ b/Engine/Source/Runtime/Mac/CoreAudio/Private/CoreAudioSource.cpp
@@ -625,7 +625,9 @@ bool FCoreAudioSoundSource::AttachToAUGraph()
 		
 		ErrorStatus = AudioConverterNew( &Buffer->PCMFormat, &AudioDevice->MatrixMixerInputFormat, &CoreAudioConverter );
 
-		bool bIs6ChannelOGG = Buffer->NumChannels == 6 && (Buffer->DecompressionState || WaveInstance->WaveData->bDecompressedFromOgg);
+		bool bIs6ChannelOGG = Buffer->NumChannels == 6
+							&& ((Buffer->DecompressionState && Buffer->DecompressionState->UsesVorbisChannelOrdering)
+							|| WaveInstance->WaveData->bDecompressedFromOgg);
 
 		AudioDevice->SetupMatrixMixerInput( DestInputNumber, bIs6ChannelOGG );
 	}
diff --git a/Engine/Source/Runtime/Windows/XAudio2/Private/XAudio2Source.cpp b/Engine/Source/Runtime/Windows/XAudio2/Private/XAudio2Source.cpp
--- a/Engine/Source/Runtime/Windows/XAudio2/Private/XAudio2Source.cpp
+++ b/Engine/Source/Runtime/Windows/XAudio2/Private/XAudio2Source.cpp
@@ -653,7 +653,8 @@ void FXAudio2SoundSource::RouteDryToSpeakers( float ChannelVolumes[CHANNELOUT_CO
 
 	case 6:
 		{
-			if (XAudio2Buffer->DecompressionState || WaveInstance->WaveData->bDecompressedFromOgg)
+			if ((XAudio2Buffer->DecompressionState && XAudio2Buffer->DecompressionState->UsesVorbisChannelOrdering())
+			||	WaveInstance->WaveData->bDecompressedFromOgg)
 			{
 				// Ordering of channels is different for 6 channel OGG
 				float SpatialisationMatrix[SPEAKER_COUNT * 6] = 
