Allowing different sample rates for Opus encoding

Choose an Opus compatible sample rate based on the input data and store that rate with the audio so that we know what to decode at.
Also changed the bit rate calculation so that it's more closely related to the original sample rate and number of channels, with the quality determining overall size reduction.

diff --git a/Engine/Source/Developer/AudioFormatOpus/Private/AudioFormatOpus.cpp b/Engine/Source/Developer/AudioFormatOpus/Private/AudioFormatOpus.cpp
--- a/Engine/Source/Developer/AudioFormatOpus/Private/AudioFormatOpus.cpp
+++ b/Engine/Source/Developer/AudioFormatOpus/Private/AudioFormatOpus.cpp
@@ -29,7 +29,7 @@ class FAudioFormatOpus : public IAudioFormat
 	enum
 	{
 		/** Version for OPUS format, this becomes part of the DDC key. */
-		UE_AUDIO_OPUS_VER = 2,
+		UE_AUDIO_OPUS_VER = 3,
 	};
 
 public:
@@ -54,8 +54,8 @@ public:
 	{
 		check(Format == NAME_OPUS);
 
-		// Constant sample rate for simplicity
-		const int32 kOpusSampleRate = 48000;
+		// Get best compatible sample rate
+		const uint16 kOpusSampleRate = GetBestOutputSampleRate(QualityInfo.SampleRate);
 		// Frame size must be one of 2.5, 5, 10, 20, 40 or 60 ms
 		const int32 kOpusFrameSizeMs = 60;
 		// Calculate frame size required by Opus
@@ -94,7 +94,8 @@ public:
 			return false;
 		}
 
-		opus_encoder_ctl(Encoder, OPUS_SET_BITRATE(GetBitRateFromQuality(QualityInfo.Quality)));
+		int32 BitRate = GetBitRateFromQuality(QualityInfo);
+		opus_encoder_ctl(Encoder, OPUS_SET_BITRATE(BitRate));
 
 		// Create a buffer to store compressed data
 		CompressedDataStore.Empty();
@@ -115,7 +116,7 @@ public:
 
 		check(QualityInfo.NumChannels <= MAX_uint8);
 		check(FramesToEncode <= MAX_uint16);
-		SerializeHeaderData(CompressedData, TrueSampleCount, QualityInfo.NumChannels, FramesToEncode);
+		SerializeHeaderData(CompressedData, kOpusSampleRate, TrueSampleCount, QualityInfo.NumChannels, FramesToEncode);
 
 		// Temporary storage with more than enough to store any compressed frame
 		TArray<uint8> TempCompressedData;
@@ -153,8 +154,8 @@ public:
 	{
 		check(Format == NAME_OPUS);
 
-		// Constant sample rate for simplicity
-		const int32 kOpusSampleRate = 48000;
+		// Get best compatible sample rate
+		const uint16 kOpusSampleRate = GetBestOutputSampleRate(QualityInfo.SampleRate);
 		// Frame size must be one of 2.5, 5, 10, 20, 40 or 60 ms
 		const int32 kOpusFrameSizeMs = 60;
 		// Calculate frame size required by Opus
@@ -227,7 +228,8 @@ public:
 			return false;
 		}
 
-		opus_multistream_encoder_ctl(Encoder, OPUS_SET_BITRATE(GetBitRateFromQuality(QualityInfo.Quality)));
+		int32 BitRate = GetBitRateFromQuality(QualityInfo);
+		opus_multistream_encoder_ctl(Encoder, OPUS_SET_BITRATE(BitRate));
 
 		// Create a buffer to store compressed data
 		CompressedDataStore.Empty();
@@ -246,7 +248,7 @@ public:
 
 		check(QualityInfo.NumChannels <= MAX_uint8);
 		check(FramesToEncode <= MAX_uint16);
-		SerializeHeaderData(CompressedData, TrueSampleCount, QualityInfo.NumChannels, FramesToEncode);
+		SerializeHeaderData(CompressedData, kOpusSampleRate, TrueSampleCount, QualityInfo.NumChannels, FramesToEncode);
 
 		// Temporary storage for source data in an interleaved format
 		TArray<uint8> TempInterleavedSrc;
@@ -361,6 +363,8 @@ public:
 			return false;
 		}
 		ReadOffset += FCStringAnsi::Strlen(OPUS_ID_STRING) + 1;
+		uint16 SampleRate = *((uint16*)(LockedSrc + ReadOffset));
+		ReadOffset += sizeof(uint16);
 		uint32 TrueSampleCount = *((uint32*)(LockedSrc + ReadOffset));
 		ReadOffset += sizeof(uint32);
 		uint8 NumChannels = *(LockedSrc + ReadOffset);
@@ -391,6 +395,34 @@ public:
 		return true;
 	}
 
+	/**
+	 * Calculate the best sample rate for the output opus data
+	 */
+	static uint16 GetBestOutputSampleRate(int32 SampleRate)
+	{
+		static const uint16 ValidSampleRates[] = 
+		{
+			0, // not really valid, but simplifies logic below
+			8000,
+			12000,
+			16000,
+			24000,
+			48000,
+		};
+
+		// look for the next highest valid rate
+		for (int32 Index = ARRAY_COUNT(ValidSampleRates) - 2; Index >= 0; Index--)
+		{
+			if (SampleRate > ValidSampleRates[Index])
+			{
+				return ValidSampleRates[Index + 1];
+			}
+		}
+		// this should never get here!
+		check(0);
+		return 0;
+	}
+
 	bool ResamplePCM(uint32 NumChannels, const TArray<uint8>& InBuffer, uint32 InSampleRate, TArray<uint8>& OutBuffer, uint32 OutSampleRate) const
 	{
 		// Initialize resampler to convert to desired rate for Opus
@@ -437,18 +469,20 @@ public:
 		return true;
 	}
 
-	int32 GetBitRateFromQuality(int32 Quality) const
+	int32 GetBitRateFromQuality(FSoundQualityInfo& QualityInfo) const
 	{
-		// Basic attempt at mapping quality to a bit rate range similar to Vorbis
-		// TODO: Not sure if range should be expanded or to make sure that it creates
-		// frequently used bit rates ie. 128000, 192000 etc.
-		return FMath::GetMappedRangeValue(FVector2D(1, 100), FVector2D(64000, 400000), Quality);
+		// There is no perfect way to map Vorbis' Quality setting to an Opus bitrate but this 
+		// will use it as a multiplier to decide how much smaller than the original the
+		// compressed data should be
+		int32 OriginalBitRate = QualityInfo.SampleRate * QualityInfo.NumChannels * SAMPLE_SIZE * 8;
+		return (float)OriginalBitRate * FMath::GetMappedRangeValue(FVector2D(1, 100), FVector2D(0.04, 0.25), QualityInfo.Quality);
 	}
 
-	void SerializeHeaderData(FMemoryWriter& CompressedData, uint32 TrueSampleCount, uint8 NumChannels, uint16 NumFrames) const
+	void SerializeHeaderData(FMemoryWriter& CompressedData, uint16 SampleRate, uint32 TrueSampleCount, uint8 NumChannels, uint16 NumFrames) const
 	{
 		const char* OpusIdentifier = OPUS_ID_STRING;
 		CompressedData.Serialize((void*)OpusIdentifier, FCStringAnsi::Strlen(OpusIdentifier) + 1);
+		CompressedData.Serialize(&SampleRate, sizeof(uint16));
 		CompressedData.Serialize(&TrueSampleCount, sizeof(uint32));
 		CompressedData.Serialize(&NumChannels, sizeof(uint8));
 		CompressedData.Serialize(&NumFrames, sizeof(uint16));
diff --git a/Engine/Source/Runtime/Engine/Private/OpusAudioInfo.cpp b/Engine/Source/Runtime/Engine/Private/OpusAudioInfo.cpp
--- a/Engine/Source/Runtime/Engine/Private/OpusAudioInfo.cpp
+++ b/Engine/Source/Runtime/Engine/Private/OpusAudioInfo.cpp
@@ -13,8 +13,7 @@
 #endif
 
 #define USE_UE4_MEM_ALLOC 1
-#define OPUS_SAMPLE_RATE 48000
-#define OPUS_MAX_FRAME_SIZE 5760
+#define OPUS_MAX_FRAME_SIZE_MS 120
 #define SAMPLE_SIZE			( ( uint32 )sizeof( short ) )
 
 ///////////////////////////////////////////////////////////////////////////////////////
@@ -44,7 +43,7 @@ FOpusDecoderWrapper
 ------------------------------------------------------------------------------------*/
 struct FOpusDecoderWrapper
 {
-	FOpusDecoderWrapper(uint8 NumChannels)
+	FOpusDecoderWrapper(uint16 SampleRate, uint8 NumChannels)
 	{
 #if WITH_OPUS
 		check(NumChannels <= 8);
@@ -52,9 +51,9 @@ struct FOpusDecoderWrapper
 	#if USE_UE4_MEM_ALLOC
 		int32 DecSize = opus_multistream_decoder_get_size(Layout.NumStreams, Layout.NumCoupledStreams);
 		Decoder = (OpusMSDecoder*)FMemory::Malloc(DecSize);
-		DecError = opus_multistream_decoder_init(Decoder, OPUS_SAMPLE_RATE, NumChannels, Layout.NumStreams, Layout.NumCoupledStreams, Layout.Mapping);
+		DecError = opus_multistream_decoder_init(Decoder, SampleRate, NumChannels, Layout.NumStreams, Layout.NumCoupledStreams, Layout.Mapping);
 	#else
-		Decoder = opus_multistream_decoder_create(OPUS_SAMPLE_RATE, NumChannels, Layout.NumStreams, Layout.NumCoupledStreams, Layout.Mapping, &DecError);
+		Decoder = opus_multistream_decoder_create(SampleRate, NumChannels, Layout.NumStreams, Layout.NumCoupledStreams, Layout.Mapping, &DecError);
 	#endif
 #endif
 	}
@@ -104,9 +103,11 @@ FOpusAudioInfo::FOpusAudioInfo(void)
 	, SrcBufferDataSize(0)
 	, SrcBufferOffset(0)
 	, AudioDataOffset(0)
+	, SampleRate(0)
 	, TrueSampleCount(0)
 	, CurrentSampleCount(0)
 	, NumChannels(0)
+	, MaxFrameSizeSamples(0)
 	, SampleStride(0)
 	, LastPCMByteSize(0)
 	, LastPCMOffset(0)
@@ -150,6 +151,7 @@ bool FOpusAudioInfo::ReadCompressedInfo(const uint8* InSrcBufferData, uint32 InS
 	{
 		SrcBufferOffset += FCStringAnsi::Strlen(OPUS_ID_STRING) + 1;
 	}
+	Read(&SampleRate, sizeof(uint16));
 	Read(&TrueSampleCount, sizeof(uint32));
 	Read(&NumChannels, sizeof(uint8));
 	uint16 SerializedFrames = 0;
@@ -158,7 +160,7 @@ bool FOpusAudioInfo::ReadCompressedInfo(const uint8* InSrcBufferData, uint32 InS
 	SampleStride = SAMPLE_SIZE * NumChannels;
 
 	check(OpusDecoderWrapper == NULL);
-	OpusDecoderWrapper = new FOpusDecoderWrapper(NumChannels);
+	OpusDecoderWrapper = new FOpusDecoderWrapper(SampleRate, NumChannels);
 	if (!OpusDecoderWrapper->WasInitialisedSuccessfully())
 	{
 		delete OpusDecoderWrapper;
@@ -166,12 +168,14 @@ bool FOpusAudioInfo::ReadCompressedInfo(const uint8* InSrcBufferData, uint32 InS
 		return false;
 	}
 
-	LastDecodedPCM.Empty(OPUS_MAX_FRAME_SIZE * SampleStride);
-	LastDecodedPCM.AddUninitialized(OPUS_MAX_FRAME_SIZE * SampleStride);
+	MaxFrameSizeSamples = (SampleRate * OPUS_MAX_FRAME_SIZE_MS) / 1000;
+
+	LastDecodedPCM.Empty(MaxFrameSizeSamples * SampleStride);
+	LastDecodedPCM.AddUninitialized(MaxFrameSizeSamples * SampleStride);
 
 	if (QualityInfo)
 	{
-		QualityInfo->SampleRate = OPUS_SAMPLE_RATE;
+		QualityInfo->SampleRate = SampleRate;
 		QualityInfo->NumChannels = NumChannels;
 		QualityInfo->SampleDataSize = TrueSampleCount * QualityInfo->NumChannels * SAMPLE_SIZE;
 		QualityInfo->Duration = (float)TrueSampleCount / QualityInfo->SampleRate;
@@ -400,7 +404,7 @@ int32 FOpusAudioInfo::DecompressToPCMBuffer(uint16 FrameSize)
 	const uint8* SrcPtr = SrcBufferData + SrcBufferOffset;
 	SrcBufferOffset += FrameSize;
 	LastPCMOffset = 0;
-	return OpusDecoderWrapper->Decode(SrcPtr, FrameSize, (int16*)LastDecodedPCM.GetTypedData(), OPUS_MAX_FRAME_SIZE);
+	return OpusDecoderWrapper->Decode(SrcPtr, FrameSize, (int16*)LastDecodedPCM.GetTypedData(), MaxFrameSizeSamples);
 }
 
 uint32 FOpusAudioInfo::IncrementCurrentSampleCount(uint32 NewSamples)
diff --git a/Engine/Source/Runtime/Engine/Public/OpusAudioInfo.h b/Engine/Source/Runtime/Engine/Public/OpusAudioInfo.h
--- a/Engine/Source/Runtime/Engine/Public/OpusAudioInfo.h
+++ b/Engine/Source/Runtime/Engine/Public/OpusAudioInfo.h
@@ -81,9 +81,11 @@ protected:
 	uint32			SrcBufferOffset;
 	uint32			AudioDataOffset;
 
+	uint16			SampleRate;
 	uint32			TrueSampleCount;
 	uint32			CurrentSampleCount;
 	uint8			NumChannels;
+	uint32			MaxFrameSizeSamples;
 	uint32			SampleStride;
 
 	TArray<uint8>	LastDecodedPCM;
