TArray documentation and a coupld of methods deprecation.

Changes mostly in Array.h. The changes in other files are only renames for deprecated functions.

diff --git a/Engine/Plugins/2D/Paper2D/Source/Paper2D/Private/AlphaBitmap.h b/Engine/Plugins/2D/Paper2D/Source/Paper2D/Private/AlphaBitmap.h
--- a/Engine/Plugins/2D/Paper2D/Source/Paper2D/Private/AlphaBitmap.h
+++ b/Engine/Plugins/2D/Paper2D/Source/Paper2D/Private/AlphaBitmap.h
@@ -47,7 +47,7 @@ struct FAlphaBitmap
 				for (int X = 0; X < Width; ++X)
 				{
 					int32 PixelByteOffset = (X + Y * Width) * BytesPerPixel;
-					const uint8* PixelPtr = TextureRawData.GetTypedData() + PixelByteOffset;
+					const uint8* PixelPtr = TextureRawData.GetData() + PixelByteOffset;
 					FColor Color;
 					if (PixelFormat == TSF_BGRA8)
 					{
diff --git a/Engine/Plugins/2D/Paper2D/Source/Paper2D/Private/PaperRenderSceneProxy.cpp b/Engine/Plugins/2D/Paper2D/Source/Paper2D/Private/PaperRenderSceneProxy.cpp
--- a/Engine/Plugins/2D/Paper2D/Source/Paper2D/Private/PaperRenderSceneProxy.cpp
+++ b/Engine/Plugins/2D/Paper2D/Source/Paper2D/Private/PaperRenderSceneProxy.cpp
@@ -429,7 +429,7 @@ void FPaperRenderSceneProxy::DrawNewBatches(FPrimitiveDrawInterface* PDI, const
 	FMeshBatch Mesh;
 
 	Mesh.UseDynamicData = true;
-	Mesh.DynamicVertexData = BatchedVertices.GetTypedData();
+	Mesh.DynamicVertexData = BatchedVertices.GetData();
 	Mesh.DynamicVertexStride = sizeof(FPaperSpriteVertex);
 
 	Mesh.VertexFactory = VertexFactory;
@@ -530,7 +530,7 @@ void FPaperRenderSceneProxy::GetNewBatchMeshes(const FSceneView* View, bool bUse
 
 			Mesh.UseDynamicData = true;
 			// Note: memory referenced by Collector must be valid as long as the Collector is around, see Collector.AllocateOneFrameResource
-			Mesh.DynamicVertexData = BatchedVertices.GetTypedData();
+			Mesh.DynamicVertexData = BatchedVertices.GetData();
 			Mesh.DynamicVertexStride = sizeof(FPaperSpriteVertex);
 
 			Mesh.VertexFactory = VertexFactory;
@@ -596,7 +596,7 @@ void FPaperRenderSceneProxy::GetBatchMesh(const FSceneView* View, bool bUseOverr
 			FMeshBatch& Mesh = Collector.AllocateMesh();
 
 			Mesh.UseDynamicData = true;
-			Mesh.DynamicVertexData = VertexArray.Vertices.GetTypedData();
+			Mesh.DynamicVertexData = VertexArray.Vertices.GetData();
 			Mesh.DynamicVertexStride = sizeof(FPaperSpriteVertex);
 			Mesh.VertexFactory = VertexFactory;
 			Mesh.bCanApplyViewModeOverrides = true;
@@ -684,7 +684,7 @@ void FPaperRenderSceneProxy::DrawBatch(FPrimitiveDrawInterface* PDI, const FScen
 			FMeshBatch Mesh;
 
 			Mesh.UseDynamicData = true;
-			Mesh.DynamicVertexData = Vertices.GetTypedData();
+			Mesh.DynamicVertexData = Vertices.GetData();
 			Mesh.DynamicVertexStride = sizeof(FPaperSpriteVertex);
 
 			Mesh.VertexFactory = VertexFactory;
diff --git a/Engine/Plugins/2D/Paper2D/Source/Paper2D/Private/PaperSpriteSceneProxy.cpp b/Engine/Plugins/2D/Paper2D/Source/Paper2D/Private/PaperSpriteSceneProxy.cpp
--- a/Engine/Plugins/2D/Paper2D/Source/Paper2D/Private/PaperSpriteSceneProxy.cpp
+++ b/Engine/Plugins/2D/Paper2D/Source/Paper2D/Private/PaperSpriteSceneProxy.cpp
@@ -153,12 +153,12 @@ void FPaperSpriteSceneProxy::SetSprite_RenderThread(const FSpriteDrawCallRecord&
 		Record.Color = NewDynamicData.Color;
 		Record.Destination = NewDynamicData.Destination;
 		Record.Texture = NewDynamicData.Texture;
-		Record.RenderVerts.Append(NewDynamicData.RenderVerts.GetTypedData(), SplitIndex);
+		Record.RenderVerts.Append(NewDynamicData.RenderVerts.GetData(), SplitIndex);
 
 		AltRecord.Color = NewDynamicData.Color;
 		AltRecord.Destination = NewDynamicData.Destination;
 		AltRecord.Texture = NewDynamicData.Texture;
-		AltRecord.RenderVerts.Append(NewDynamicData.RenderVerts.GetTypedData() + SplitIndex, NewDynamicData.RenderVerts.Num() - SplitIndex);
+		AltRecord.RenderVerts.Append(NewDynamicData.RenderVerts.GetData() + SplitIndex, NewDynamicData.RenderVerts.Num() - SplitIndex);
 	}
 	else
 	{
diff --git a/Engine/Plugins/2D/Paper2D/Source/Paper2DEditor/Private/Atlasing/PaperAtlasGenerator.cpp b/Engine/Plugins/2D/Paper2D/Source/Paper2DEditor/Private/Atlasing/PaperAtlasGenerator.cpp
--- a/Engine/Plugins/2D/Paper2D/Source/Paper2DEditor/Private/Atlasing/PaperAtlasGenerator.cpp
+++ b/Engine/Plugins/2D/Paper2D/Source/Paper2DEditor/Private/Atlasing/PaperAtlasGenerator.cpp
@@ -38,7 +38,7 @@ public:
 	{
 		const int32 NumMips = 1;
 
-		Texture->Source.Init(AtlasWidth, AtlasHeight, /*NumSlices=*/ 1, NumMips, ETextureSourceFormat::TSF_BGRA8, AtlasData.GetTypedData());
+		Texture->Source.Init(AtlasWidth, AtlasHeight, /*NumSlices=*/ 1, NumMips, ETextureSourceFormat::TSF_BGRA8, AtlasData.GetData());
 		Texture->UpdateResource();
 	}
 
@@ -58,7 +58,7 @@ public:
 		{
 			uint32 BytesPerPixel = SourceData.GetBytesPerPixel();
 			uint8* OffsetSource = SourceData.LockMip(0) + (FMath::TruncToInt(Sprite->GetSourceUV().X) + FMath::TruncToInt(Sprite->GetSourceUV().Y) * SourceData.GetSizeX()) * BytesPerPixel;
-			uint8* OffsetDest = DummyBuffer.GetTypedData();
+			uint8* OffsetDest = DummyBuffer.GetData();
 
 			CopyTextureData(OffsetSource, OffsetDest, SpriteSize.X, SpriteSize.Y, BytesPerPixel, SourceData.GetSizeX() * BytesPerPixel, SpriteSize.X * BytesPerPixel);
 
diff --git a/Engine/Plugins/Editor/SpeedTreeImporter/Source/SpeedTreeImporter/Private/SpeedTreeImportFactory.cpp b/Engine/Plugins/Editor/SpeedTreeImporter/Source/SpeedTreeImporter/Private/SpeedTreeImportFactory.cpp
--- a/Engine/Plugins/Editor/SpeedTreeImporter/Source/SpeedTreeImporter/Private/SpeedTreeImportFactory.cpp
+++ b/Engine/Plugins/Editor/SpeedTreeImporter/Source/SpeedTreeImporter/Private/SpeedTreeImportFactory.cpp
@@ -411,7 +411,7 @@ static UTexture* CreateSpeedTreeMaterialTexture(UObject* Parent, FString Filenam
 			TextureFact->CompressionSettings = TC_Normalmap;
 		}
 
-		const uint8* PtrTexture = TextureData.GetTypedData();
+		const uint8* PtrTexture = TextureData.GetData();
 		UnrealTexture = (UTexture*)TextureFact->FactoryCreateBinary(UTexture2D::StaticClass(), Package, *TextureName, RF_Standalone|RF_Public, NULL, *Extension, PtrTexture, PtrTexture + TextureData.Num(), GWarn);
 		if (UnrealTexture != NULL)
 		{
diff --git a/Engine/Plugins/Messaging/UdpMessaging/Source/UdpMessaging/Private/Transport/ReassembledUdpMessage.h b/Engine/Plugins/Messaging/UdpMessaging/Source/UdpMessaging/Private/Transport/ReassembledUdpMessage.h
--- a/Engine/Plugins/Messaging/UdpMessaging/Source/UdpMessaging/Private/Transport/ReassembledUdpMessage.h
+++ b/Engine/Plugins/Messaging/UdpMessaging/Source/UdpMessaging/Private/Transport/ReassembledUdpMessage.h
@@ -135,7 +135,7 @@ public:
 		{
 			if (SegmentOffset + SegmentData.Num() <= Data.Num())
 			{
-				FMemory::Memcpy(Data.GetTypedData() + SegmentOffset, SegmentData.GetTypedData(), SegmentData.Num());
+				FMemory::Memcpy(Data.GetData() + SegmentOffset, SegmentData.GetData(), SegmentData.Num());
 
 				PendingSegments[SegmentNumber] = false;
 				--PendingSegmentsCount;
diff --git a/Engine/Plugins/Messaging/UdpMessaging/Source/UdpMessaging/Private/Transport/UdpMessageSegmenter.h b/Engine/Plugins/Messaging/UdpMessaging/Source/UdpMessaging/Private/Transport/UdpMessageSegmenter.h
--- a/Engine/Plugins/Messaging/UdpMessaging/Source/UdpMessaging/Private/Transport/UdpMessageSegmenter.h
+++ b/Engine/Plugins/Messaging/UdpMessaging/Source/UdpMessaging/Private/Transport/UdpMessageSegmenter.h
@@ -89,7 +89,7 @@ public:
 			OutData.AddUninitialized(ActualSegmentSize);
 
 			MessageReader->Seek(SegmentOffset);
-			MessageReader->Serialize(OutData.GetTypedData(), ActualSegmentSize);
+			MessageReader->Serialize(OutData.GetData(), ActualSegmentSize);
 
 			//FMemory::Memcpy(OutData.GetTypedData(), Message->GetTypedData() + SegmentOffset, ActualSegmentSize);
 
diff --git a/Engine/Plugins/Messaging/UdpMessaging/Source/UdpMessaging/Private/Tunnel/UdpMessageTunnel.h b/Engine/Plugins/Messaging/UdpMessaging/Source/UdpMessaging/Private/Tunnel/UdpMessageTunnel.h
--- a/Engine/Plugins/Messaging/UdpMessaging/Source/UdpMessaging/Private/Tunnel/UdpMessageTunnel.h
+++ b/Engine/Plugins/Messaging/UdpMessaging/Source/UdpMessaging/Private/Tunnel/UdpMessageTunnel.h
@@ -259,7 +259,7 @@ protected:
 				}
 
 				int32 BytesSent = 0;
-				UnicastSocket->SendTo(Payload->GetTypedData(), Payload->Num(), BytesSent, *RecipientEndpoint.ToInternetAddr());
+				UnicastSocket->SendTo(Payload->GetData(), Payload->Num(), BytesSent, *RecipientEndpoint.ToInternetAddr());
 			}
 			else
 			{
diff --git a/Engine/Plugins/Messaging/UdpMessaging/Source/UdpMessaging/Private/Tunnel/UdpMessageTunnelConnection.h b/Engine/Plugins/Messaging/UdpMessaging/Source/UdpMessaging/Private/Tunnel/UdpMessageTunnelConnection.h
--- a/Engine/Plugins/Messaging/UdpMessaging/Source/UdpMessaging/Private/Tunnel/UdpMessageTunnelConnection.h
+++ b/Engine/Plugins/Messaging/UdpMessaging/Source/UdpMessaging/Private/Tunnel/UdpMessageTunnelConnection.h
@@ -191,7 +191,7 @@ protected:
 			PayloadSizeData.Init(sizeof(uint16));
 			
 			// ... read it from the stream without removing it...
-			if (!Socket->Recv(PayloadSizeData.GetTypedData(), sizeof(uint16), BytesRead, ESocketReceiveFlags::Peek))
+			if (!Socket->Recv(PayloadSizeData.GetData(), sizeof(uint16), BytesRead, ESocketReceiveFlags::Peek))
 			{
 				return false;
 			}
@@ -254,7 +254,7 @@ protected:
 				uint16 PayloadSize = Payload->Num();
 				PayloadSizeData << PayloadSize;
 
-				if (!Socket->Send(PayloadSizeData.GetTypedData(), sizeof(uint16), BytesSent))
+				if (!Socket->Send(PayloadSizeData.GetData(), sizeof(uint16), BytesSent))
 				{
 					return false;
 				}
@@ -262,7 +262,7 @@ protected:
 				TotalBytesSent += BytesSent;
 
 				// send the payload
-				if (!Socket->Send(Payload->GetTypedData(), Payload->Num(), BytesSent))
+				if (!Socket->Send(Payload->GetData(), Payload->Num(), BytesSent))
 				{
 					return false;
 				}
diff --git a/Engine/Plugins/Runtime/CustomMeshComponent/Source/CustomMeshComponent/Private/CustomMeshComponent.cpp b/Engine/Plugins/Runtime/CustomMeshComponent/Source/CustomMeshComponent/Private/CustomMeshComponent.cpp
--- a/Engine/Plugins/Runtime/CustomMeshComponent/Source/CustomMeshComponent/Private/CustomMeshComponent.cpp
+++ b/Engine/Plugins/Runtime/CustomMeshComponent/Source/CustomMeshComponent/Private/CustomMeshComponent.cpp
@@ -17,7 +17,7 @@ public:
 
 		// Copy the vertex data into the vertex buffer.
 		void* VertexBufferData = RHILockVertexBuffer(VertexBufferRHI,0,Vertices.Num() * sizeof(FDynamicMeshVertex), RLM_WriteOnly);
-		FMemory::Memcpy(VertexBufferData,Vertices.GetTypedData(),Vertices.Num() * sizeof(FDynamicMeshVertex));
+		FMemory::Memcpy(VertexBufferData,Vertices.GetData(),Vertices.Num() * sizeof(FDynamicMeshVertex));
 		RHIUnlockVertexBuffer(VertexBufferRHI);
 	}
 
@@ -36,7 +36,7 @@ public:
 
 		// Write the indices to the index buffer.
 		void* Buffer = RHILockIndexBuffer(IndexBufferRHI,0,Indices.Num() * sizeof(int32),RLM_WriteOnly);
-		FMemory::Memcpy(Buffer,Indices.GetTypedData(),Indices.Num() * sizeof(int32));
+		FMemory::Memcpy(Buffer,Indices.GetData(),Indices.Num() * sizeof(int32));
 		RHIUnlockIndexBuffer(IndexBufferRHI);
 	}
 };
diff --git a/Engine/Source/Developer/AudioFormatADPCM/Private/AudioFormatADPCM.cpp b/Engine/Source/Developer/AudioFormatADPCM/Private/AudioFormatADPCM.cpp
--- a/Engine/Source/Developer/AudioFormatADPCM/Private/AudioFormatADPCM.cpp
+++ b/Engine/Source/Developer/AudioFormatADPCM/Private/AudioFormatADPCM.cpp
@@ -195,7 +195,7 @@ namespace LPCM
 
 		RiffDataChunks[1].ID = UE_MAKEFOURCC('d','a','t','a');
 		RiffDataChunks[1].DataSize = InputPCMData.Num();
-		RiffDataChunks[1].Data = const_cast<uint8*>(InputPCMData.GetTypedData());
+		RiffDataChunks[1].Data = const_cast<uint8*>(InputPCMData.GetData());
 
 		GenerateWaveFile(RiffDataChunks, CompressedDataStore);
 	}
diff --git a/Engine/Source/Developer/AudioFormatOgg/Private/AudioFormatOgg.cpp b/Engine/Source/Developer/AudioFormatOgg/Private/AudioFormatOgg.cpp
--- a/Engine/Source/Developer/AudioFormatOgg/Private/AudioFormatOgg.cpp
+++ b/Engine/Source/Developer/AudioFormatOgg/Private/AudioFormatOgg.cpp
@@ -123,7 +123,7 @@ public:
 			{
 				// Read samples
 				uint32 BytesToRead = FMath::Min( SAMPLES_TO_READ * QualityInfo.NumChannels * SAMPLE_SIZE, QualityInfo.SampleDataSize - BufferOffset );
-				FMemory::Memcpy( ReadBuffer, SrcBuffer.GetTypedData() + BufferOffset, BytesToRead );
+				FMemory::Memcpy( ReadBuffer, SrcBuffer.GetData() + BufferOffset, BytesToRead );
 				BufferOffset += BytesToRead;
 
 				if( BytesToRead == 0)
@@ -308,7 +308,7 @@ public:
 					uint32 i = 0;
 					for( j = 0; j < ChannelOrder.Num(); j++ )
 					{
-						short* ReadBuffer = ( short* )( SrcBuffers[ChannelOrder[j]].GetTypedData() + BufferOffset );
+						short* ReadBuffer = (short*)(SrcBuffers[ChannelOrder[j]].GetData() + BufferOffset);
 						for( i = 0; i < BytesToRead / SAMPLE_SIZE; i++ )
 						{
 							buffer[j][i] = ( ReadBuffer[i] ) / 32768.0f;
@@ -421,7 +421,7 @@ public:
 		}
 
 		// Parse the audio header for the relevant information
-		if (!AudioInfo.ReadCompressedInfo(CompressedDataStore.GetTypedData(), CompressedDataStore.Num(), &QualityInfo))
+		if (!AudioInfo.ReadCompressedInfo(CompressedDataStore.GetData(), CompressedDataStore.Num(), &QualityInfo))
 		{
 			return 0;
 		}
@@ -429,7 +429,7 @@ public:
 		// Decompress all the sample data
 		OutBuffer.Empty(QualityInfo.SampleDataSize);
 		OutBuffer.AddZeroed(QualityInfo.SampleDataSize);
-		AudioInfo.ExpandFile(OutBuffer.GetTypedData(), &QualityInfo);
+		AudioInfo.ExpandFile(OutBuffer.GetData(), &QualityInfo);
 
 		return CompressedDataStore.Num();
 	}
diff --git a/Engine/Source/Developer/AudioFormatOpus/Private/AudioFormatOpus.cpp b/Engine/Source/Developer/AudioFormatOpus/Private/AudioFormatOpus.cpp
--- a/Engine/Source/Developer/AudioFormatOpus/Private/AudioFormatOpus.cpp
+++ b/Engine/Source/Developer/AudioFormatOpus/Private/AudioFormatOpus.cpp
@@ -124,7 +124,7 @@ public:
 
 		while (SrcBufferOffset < SrcBufferCopy.Num())
 		{
-			int32 CompressedLength = opus_encode(Encoder, (const opus_int16*)(SrcBufferCopy.GetTypedData() + SrcBufferOffset), kOpusFrameSizeSamples, TempCompressedData.GetTypedData(), TempCompressedData.Num());
+			int32 CompressedLength = opus_encode(Encoder, (const opus_int16*)(SrcBufferCopy.GetData() + SrcBufferOffset), kOpusFrameSizeSamples, TempCompressedData.GetData(), TempCompressedData.Num());
 
 			if (CompressedLength < 0)
 			{
@@ -140,7 +140,7 @@ public:
 			{
 				// Store frame length and copy compressed data before incrementing pointers
 				check(CompressedLength < MAX_uint16);
-				SerialiseFrameData(CompressedData, TempCompressedData.GetTypedData(), CompressedLength);
+				SerialiseFrameData(CompressedData, TempCompressedData.GetData(), CompressedLength);
 				SrcBufferOffset += kBytesPerFrame;
 			}
 		}
@@ -218,9 +218,9 @@ public:
 #if USE_UE4_MEM_ALLOC
 		int32 EncSize = opus_multistream_surround_encoder_get_size(QualityInfo.NumChannels, mapping_family);
 		Encoder = (OpusMSEncoder*)FMemory::Malloc(EncSize);
-		EncError = opus_multistream_surround_encoder_init(Encoder, kOpusSampleRate, QualityInfo.NumChannels, mapping_family, &streams, &coupled_streams, mapping.GetTypedData(), OPUS_APPLICATION_AUDIO);
+		EncError = opus_multistream_surround_encoder_init(Encoder, kOpusSampleRate, QualityInfo.NumChannels, mapping_family, &streams, &coupled_streams, mapping.GetData(), OPUS_APPLICATION_AUDIO);
 #else
-		Encoder = opus_multistream_surround_encoder_create(kOpusSampleRate, QualityInfo.NumChannels, mapping_family, &streams, &coupled_streams, mapping.GetTypedData(), OPUS_APPLICATION_AUDIO, &EncError);
+		Encoder = opus_multistream_surround_encoder_create(kOpusSampleRate, QualityInfo.NumChannels, mapping_family, &streams, &coupled_streams, mapping.GetData(), OPUS_APPLICATION_AUDIO, &EncError);
 #endif
 		if (EncError != OPUS_OK)
 		{
@@ -281,12 +281,12 @@ public:
 				else
 				{
 					// Zero the rest of the temp buffer to make it an exact frame
-					FMemory::Memzero(TempInterleavedSrc.GetTypedData() + CurrInterleavedOffset, kBytesPerFrame - CurrInterleavedOffset);
+					FMemory::Memzero(TempInterleavedSrc.GetData() + CurrInterleavedOffset, kBytesPerFrame - CurrInterleavedOffset);
 					SampleIndex = kOpusFrameSizeSamples;
 				}
 			}
 
-			int32 CompressedLength = opus_multistream_encode(Encoder, (const opus_int16*)(TempInterleavedSrc.GetTypedData()), kOpusFrameSizeSamples, TempCompressedData.GetTypedData(), TempCompressedData.Num());
+			int32 CompressedLength = opus_multistream_encode(Encoder, (const opus_int16*)(TempInterleavedSrc.GetData()), kOpusFrameSizeSamples, TempCompressedData.GetData(), TempCompressedData.Num());
 
 			if (CompressedLength < 0)
 			{
@@ -302,7 +302,7 @@ public:
 			{
 				// Store frame length and copy compressed data before incrementing pointers
 				check(CompressedLength < MAX_uint16);
-				SerialiseFrameData(CompressedData, TempCompressedData.GetTypedData(), CompressedLength);
+				SerialiseFrameData(CompressedData, TempCompressedData.GetData(), CompressedLength);
 				SrcBufferOffset += kOpusFrameSizeSamples * SAMPLE_SIZE;
 			}
 		}
@@ -330,7 +330,7 @@ public:
 		}
 
 		// Parse the opus header for the relevant information
-		if( !AudioInfo.ReadCompressedInfo( CompressedDataStore.GetTypedData(), CompressedDataStore.Num(), &QualityInfo ) )
+		if( !AudioInfo.ReadCompressedInfo( CompressedDataStore.GetData(), CompressedDataStore.Num(), &QualityInfo ) )
 		{
 			return 0;
 		}
@@ -338,7 +338,7 @@ public:
 		// Decompress all the sample data
 		OutBuffer.Empty(QualityInfo.SampleDataSize);
 		OutBuffer.AddZeroed(QualityInfo.SampleDataSize);
-		AudioInfo.ExpandFile( OutBuffer.GetTypedData(), &QualityInfo );
+		AudioInfo.ExpandFile( OutBuffer.GetData(), &QualityInfo );
 
 		return CompressedDataStore.Num();
 	}
@@ -355,7 +355,7 @@ public:
 		uint32 ReadOffset = 0;
 		uint32 WriteOffset = 0;
 		uint16 ProcessedFrames = 0;
-		const uint8* LockedSrc = SrcBuffer.GetTypedData();
+		const uint8* LockedSrc = SrcBuffer.GetData();
 
 		// Read Identifier, True Sample Count, Number of channels and Frames to Encode first
 		if (FCStringAnsi::Strcmp((char*)LockedSrc, OPUS_ID_STRING) != 0)
@@ -446,11 +446,11 @@ public:
 		// Do resampling and check results
 		if (NumChannels == 1)
 		{
-			err = speex_resampler_process_int(resampler, 0, (const short*)(InBuffer.GetTypedData()), &InSamples, (short*)(OutBuffer.GetTypedData()), &OutSamples);
+			err = speex_resampler_process_int(resampler, 0, (const short*)(InBuffer.GetData()), &InSamples, (short*)(OutBuffer.GetData()), &OutSamples);
 		}
 		else
 		{
-			err = speex_resampler_process_interleaved_int(resampler, (const short*)(InBuffer.GetTypedData()), &InSamples, (short*)(OutBuffer.GetTypedData()), &OutSamples);
+			err = speex_resampler_process_interleaved_int(resampler, (const short*)(InBuffer.GetData()), &InSamples, (short*)(OutBuffer.GetData()), &OutSamples);
 		}
 
 		speex_resampler_destroy(resampler);
@@ -525,7 +525,7 @@ public:
 		TArray<uint8>& NewBuffer = *new (OutBuffers) TArray<uint8>;
 		NewBuffer.Empty(ChunkSize);
 		NewBuffer.AddUninitialized(ChunkSize);
-		FMemory::Memcpy(NewBuffer.GetTypedData(), ChunkData, ChunkSize);
+		FMemory::Memcpy(NewBuffer.GetData(), ChunkData, ChunkSize);
 		return ChunkSize;
 	}
 };
diff --git a/Engine/Source/Developer/CrashTracker/Private/CrashVideoCapture.cpp b/Engine/Source/Developer/CrashTracker/Private/CrashVideoCapture.cpp
--- a/Engine/Source/Developer/CrashTracker/Private/CrashVideoCapture.cpp
+++ b/Engine/Source/Developer/CrashTracker/Private/CrashVideoCapture.cpp
@@ -61,7 +61,7 @@ void FAsyncImageCompress::DoWork()
 	{
 		DataFrame->ActualSize = CompressedSize;
 
-		FMemory::Memcpy(DataFrame->Data, CompressedData.GetTypedData(), CompressedSize);
+		FMemory::Memcpy(DataFrame->Data, CompressedData.GetData(), CompressedSize);
 	}
 	else
 	{
diff --git a/Engine/Source/Developer/DerivedDataCache/Private/DerivedDataBackendCorruptionWrapper.h b/Engine/Source/Developer/DerivedDataCache/Private/DerivedDataBackendCorruptionWrapper.h
--- a/Engine/Source/Developer/DerivedDataCache/Private/DerivedDataBackendCorruptionWrapper.h
+++ b/Engine/Source/Developer/DerivedDataCache/Private/DerivedDataBackendCorruptionWrapper.h
@@ -39,7 +39,7 @@ struct FDerivedDataTrailer
 	FDerivedDataTrailer(const TArray<uint8>& Data)
 		: Magic(MagicConstant)
 		, Version(1)
-		, CRCofPayload(FCrc::MemCrc_DEPRECATED(Data.GetTypedData(),Data.Num()))
+		, CRCofPayload(FCrc::MemCrc_DEPRECATED(Data.GetData(), Data.Num()))
 		, SizeOfPayload(Data.Num())
 	{
 	}
diff --git a/Engine/Source/Developer/DerivedDataCache/Private/PakFileDerivedDataBackend.h b/Engine/Source/Developer/DerivedDataCache/Private/PakFileDerivedDataBackend.h
--- a/Engine/Source/Developer/DerivedDataCache/Private/PakFileDerivedDataBackend.h
+++ b/Engine/Source/Developer/DerivedDataCache/Private/PakFileDerivedDataBackend.h
@@ -118,8 +118,8 @@ public:
 				check(!OutData.Num());
 				check(FileHandle->IsLoading());
 				OutData.AddUninitialized(Item->Size);
-				FileHandle->Serialize(OutData.GetTypedData(),  int64(Item->Size));
-				uint32 TestCrc = FCrc::MemCrc_DEPRECATED(OutData.GetTypedData(), Item->Size);
+				FileHandle->Serialize(OutData.GetData(), int64(Item->Size));
+				uint32 TestCrc = FCrc::MemCrc_DEPRECATED(OutData.GetData(), Item->Size);
 				if (TestCrc != Item->Crc)
 				{
 					UE_LOG(LogDerivedDataCache, Warning, TEXT("Pak file, bad crc."));
@@ -160,7 +160,7 @@ public:
 			{
 				check(InData.Num());
 				check(Key.Len());
-				uint32 Crc = FCrc::MemCrc_DEPRECATED(InData.GetTypedData(), InData.Num());
+				uint32 Crc = FCrc::MemCrc_DEPRECATED(InData.GetData(), InData.Num());
 				check(FileHandle);
 				check(FileHandle->IsSaving());
 				int64 Offset = FileHandle->Tell();
@@ -172,7 +172,7 @@ public:
 				}
 				else
 				{
-					FileHandle->Serialize(InData.GetTypedData(), int64(InData.Num()));
+					FileHandle->Serialize(InData.GetData(), int64(InData.Num()));
 					UE_LOG(LogDerivedDataCache, Verbose, TEXT("FPakFileDerivedDataBackend: Put %s"), CacheKey);
 					CacheItems.Add(Key,FCacheValue(Offset, InData.Num(), Crc));
 				}
@@ -228,7 +228,7 @@ public:
 			}
 			check(NumProcessed == NumItems);
 		}
-		uint32 IndexCrc = FCrc::MemCrc_DEPRECATED(IndexBuffer.GetTypedData(), IndexBuffer.Num());
+		uint32 IndexCrc = FCrc::MemCrc_DEPRECATED(IndexBuffer.GetData(), IndexBuffer.Num());
 		uint32 SizeIndex = uint32(IndexBuffer.Num());
 
 		uint32 Magic = PakCache_Magic;
@@ -238,10 +238,10 @@ public:
 		Saver << IndexCrc;
 		Saver << NumItems;
 		Saver << SizeIndex;
-		Saver.Serialize(IndexBuffer.GetTypedData(), IndexBuffer.Num());
+		Saver.Serialize(IndexBuffer.GetData(), IndexBuffer.Num());
 		Saver << Magic;
 		Saver << IndexOffset;
-		FileHandle->Serialize(Buffer.GetTypedData(), Buffer.Num());
+		FileHandle->Serialize(Buffer.GetData(), Buffer.Num());
 		CacheItems.Empty();
 		FileHandle.Reset();
 		bClosed = true;
@@ -276,7 +276,7 @@ public:
 			}
 			check(Trailer >= 0 && Trailer < FileSize);
 			Buffer.AddUninitialized(sizeof(int64) + sizeof(uint32));
-			FileHandle->Serialize(Buffer.GetTypedData(), int64(sizeof(int64) + sizeof(uint32)));
+			FileHandle->Serialize(Buffer.GetData(), int64(sizeof(int64)+sizeof(uint32)));
 			FMemoryReader Loader(Buffer);
 			uint32 Magic = 0;
 			Loader << Magic;
@@ -299,7 +299,7 @@ public:
 				return false;
 			}
 			Buffer.AddUninitialized(sizeof(uint32) * 4);
-			FileHandle->Serialize(Buffer.GetTypedData(), int64(sizeof(uint32) * 4));
+			FileHandle->Serialize(Buffer.GetData(), int64(sizeof(uint32)* 4));
 			FMemoryReader Loader(Buffer);
 			uint32 Magic = 0;
 			Loader << Magic;
@@ -320,7 +320,7 @@ public:
 		{
 			TArray<uint8> Buffer;
 			Buffer.AddUninitialized(SizeIndex);
-			FileHandle->Serialize(Buffer.GetTypedData(), int64(SizeIndex));
+			FileHandle->Serialize(Buffer.GetData(), int64(SizeIndex));
 			FMemoryReader Loader(Buffer);
 			while (Loader.Tell() < (int32)SizeIndex)
 			{
diff --git a/Engine/Source/Developer/DesktopPlatform/Private/Windows/DesktopPlatformWindows.cpp b/Engine/Source/Developer/DesktopPlatform/Private/Windows/DesktopPlatformWindows.cpp
--- a/Engine/Source/Developer/DesktopPlatform/Private/Windows/DesktopPlatformWindows.cpp
+++ b/Engine/Source/Developer/DesktopPlatform/Private/Windows/DesktopPlatformWindows.cpp
@@ -608,20 +608,20 @@ void FDesktopPlatformWindows::GetRequiredRegistrySettings(TIndirectArray<FRegist
 
 	// HKCU\SOFTWARE\Classes\.uproject
 	FRegistryRootedKey *UserRootExtensionKey = new FRegistryRootedKey(HKEY_CURRENT_USER, TEXT("SOFTWARE\\Classes\\.uproject"));
-	RootedKeys.AddRawItem(UserRootExtensionKey);
+	RootedKeys.Add(UserRootExtensionKey);
 
 	// HKLM\SOFTWARE\Classes\.uproject
 	FRegistryRootedKey *RootExtensionKey = new FRegistryRootedKey(HKEY_LOCAL_MACHINE, TEXT("SOFTWARE\\Classes\\.uproject"));
 	RootExtensionKey->Key = new FRegistryKey();
 	RootExtensionKey->Key->SetValue(TEXT(""), TEXT("Unreal.ProjectFile"));
-	RootedKeys.AddRawItem(RootExtensionKey);
+	RootedKeys.Add(RootExtensionKey);
 
 	// HKLM\SOFTWARE\Classes\Unreal.ProjectFile
 	FRegistryRootedKey *RootFileTypeKey = new FRegistryRootedKey(HKEY_LOCAL_MACHINE, TEXT("SOFTWARE\\Classes\\Unreal.ProjectFile"));
 	RootFileTypeKey->Key = new FRegistryKey();
 	RootFileTypeKey->Key->SetValue(TEXT(""), TEXT("Unreal Engine Project File"));
 	RootFileTypeKey->Key->FindOrAddKey(L"DefaultIcon")->SetValue(TEXT(""), QuotedExecutableFileName);
-	RootedKeys.AddRawItem(RootFileTypeKey);
+	RootedKeys.Add(RootFileTypeKey);
 
 	// HKLM\SOFTWARE\Classes\Unreal.ProjectFile\shell
 	FRegistryKey *ShellKey = RootFileTypeKey->Key->FindOrAddKey(TEXT("shell"));
@@ -668,7 +668,7 @@ void FDesktopPlatformWindows::GetRequiredRegistrySettings(TIndirectArray<FRegist
 	}
 	if(bDeleteUserChoiceKey)
 	{
-		RootedKeys.AddRawItem(new FRegistryRootedKey(HKEY_CURRENT_USER, UserChoicePath));
+		RootedKeys.Add(new FRegistryRootedKey(HKEY_CURRENT_USER, UserChoicePath));
 	}
 }
 
diff --git a/Engine/Source/Developer/DirectoryWatcher/Private/Windows/DirectoryWatcherWindows.cpp b/Engine/Source/Developer/DirectoryWatcher/Private/Windows/DirectoryWatcherWindows.cpp
--- a/Engine/Source/Developer/DirectoryWatcher/Private/Windows/DirectoryWatcherWindows.cpp
+++ b/Engine/Source/Developer/DirectoryWatcher/Private/Windows/DirectoryWatcherWindows.cpp
@@ -133,7 +133,7 @@ void FDirectoryWatcherWindows::Tick( float DeltaSeconds )
 	// Trigger any file changed delegates that are queued up
 	if ( DirectoryHandles.Num() > 0 )
 	{
-		MsgWaitForMultipleObjectsEx(DirectoryHandles.Num(), DirectoryHandles.GetTypedData(), 0, QS_ALLEVENTS, MWMO_ALERTABLE);
+		MsgWaitForMultipleObjectsEx(DirectoryHandles.Num(), DirectoryHandles.GetData(), 0, QS_ALLEVENTS, MWMO_ALERTABLE);
 	}
 
 	// Delete any stale or invalid requests
diff --git a/Engine/Source/Developer/FunctionalTesting/Private/FunctionalTest.cpp b/Engine/Source/Developer/FunctionalTesting/Private/FunctionalTest.cpp
--- a/Engine/Source/Developer/FunctionalTesting/Private/FunctionalTest.cpp
+++ b/Engine/Source/Developer/FunctionalTesting/Private/FunctionalTest.cpp
@@ -106,7 +106,7 @@ void AFunctionalTest::FinishTest(TEnumAsByte<EFunctionalTestResult::Type> TestRe
 
 	OnTestFinished.Broadcast();
 
-	AActor** ActorToDestroy = AutoDestroyActors.GetTypedData();
+	AActor** ActorToDestroy = AutoDestroyActors.GetData();
 
 	for (int32 ActorIndex = 0; ActorIndex < AutoDestroyActors.Num(); ++ActorIndex, ++ActorToDestroy)
 	{
diff --git a/Engine/Source/Developer/GameplayDebugger/Private/GameplayDebuggingComponent.cpp b/Engine/Source/Developer/GameplayDebugger/Private/GameplayDebuggingComponent.cpp
--- a/Engine/Source/Developer/GameplayDebugger/Private/GameplayDebuggingComponent.cpp
+++ b/Engine/Source/Developer/GameplayDebugger/Private/GameplayDebuggingComponent.cpp
@@ -434,7 +434,7 @@ void UGameplayDebuggingComponent::CollectBehaviorTreeData()
 
 				const int32 UncompressedSize = UncompressedBuffer.Num();
 				int32 CompressedSize = BlackboardRepData.Num() - HeaderSize;
-				uint8* DestBuffer = BlackboardRepData.GetTypedData();
+				uint8* DestBuffer = BlackboardRepData.GetData();
 				FMemory::Memcpy(DestBuffer, &UncompressedSize, HeaderSize);
 				DestBuffer += HeaderSize;
 
@@ -454,14 +454,14 @@ void UGameplayDebuggingComponent::OnRep_UpdateBlackboard()
 		TArray<uint8> UncompressedBuffer;
 		int32 UncompressedSize = 0;
 		const int32 HeaderSize = sizeof(int32);
-		uint8* SrcBuffer = (uint8*)BlackboardRepData.GetTypedData();
+		uint8* SrcBuffer = (uint8*)BlackboardRepData.GetData();
 		FMemory::Memcpy(&UncompressedSize, SrcBuffer, HeaderSize);
 		SrcBuffer += HeaderSize;
 		const int32 CompressedSize = BlackboardRepData.Num() - HeaderSize;
 
 		UncompressedBuffer.AddZeroed(UncompressedSize);
 
-		FCompression::UncompressMemory((ECompressionFlags)(COMPRESS_ZLIB), (void*)UncompressedBuffer.GetTypedData(), UncompressedSize, SrcBuffer, CompressedSize);
+		FCompression::UncompressMemory((ECompressionFlags)(COMPRESS_ZLIB), (void*)UncompressedBuffer.GetData(), UncompressedSize, SrcBuffer, CompressedSize);
 		FMemoryReader ArReader(UncompressedBuffer);
 
 		ArReader << BlackboardString;
@@ -584,14 +584,14 @@ void UGameplayDebuggingComponent::OnRep_UpdateEQS()
 		TArray<uint8> UncompressedBuffer;
 		int32 UncompressedSize = 0;
 		const int32 HeaderSize = sizeof(int32);
-		uint8* SrcBuffer = (uint8*)EQSRepData.GetTypedData();
+		uint8* SrcBuffer = (uint8*)EQSRepData.GetData();
 		FMemory::Memcpy(&UncompressedSize, SrcBuffer, HeaderSize);
 		SrcBuffer += HeaderSize;
 		const int32 CompressedSize = EQSRepData.Num() - HeaderSize;
 
 		UncompressedBuffer.AddZeroed(UncompressedSize);
 
-		FCompression::UncompressMemory((ECompressionFlags)(COMPRESS_ZLIB), (void*)UncompressedBuffer.GetTypedData(), UncompressedSize, SrcBuffer, CompressedSize);
+		FCompression::UncompressMemory((ECompressionFlags)(COMPRESS_ZLIB), (void*)UncompressedBuffer.GetData(), UncompressedSize, SrcBuffer, CompressedSize);
 		FMemoryReader ArReader(UncompressedBuffer);
 
 		ArReader << EQSLocalData;
@@ -675,7 +675,7 @@ void UGameplayDebuggingComponent::CollectEQSData()
 	EQSRepData.Init(0, HeaderSize + FMath::TruncToInt(1.1f * UncompressedSize));
 
 	int32 CompressedSize = EQSRepData.Num() - HeaderSize;
-	uint8* DestBuffer = EQSRepData.GetTypedData();
+	uint8* DestBuffer = EQSRepData.GetData();
 	FMemory::Memcpy(DestBuffer, &UncompressedSize, HeaderSize);
 	DestBuffer += HeaderSize;
 
@@ -968,7 +968,7 @@ void UGameplayDebuggingComponent::ServerCollectNavmeshData_Implementation(FVecto
 
 	const int32 UncompressedSize = UncompressedBuffer.Num();
 	int32 CompressedSize = NavmeshRepData.Num() - HeaderSize;
-	uint8* DestBuffer = NavmeshRepData.GetTypedData();
+	uint8* DestBuffer = NavmeshRepData.GetData();
 	FMemory::Memcpy(DestBuffer, &UncompressedSize, HeaderSize);
 	DestBuffer += HeaderSize;
 
@@ -1004,7 +1004,7 @@ void UGameplayDebuggingComponent::PrepareNavMeshData(struct FNavMeshSceneProxyDa
 		if (NavmeshRepData.Num() > HeaderSize)
 		{
 			int32 UncompressedSize = 0;
-			uint8* SrcBuffer = (uint8*)NavmeshRepData.GetTypedData();
+			uint8* SrcBuffer = (uint8*)NavmeshRepData.GetData();
 			FMemory::Memcpy(&UncompressedSize, SrcBuffer, HeaderSize);
 			SrcBuffer += HeaderSize;
 			const int32 CompressedSize = NavmeshRepData.Num() - HeaderSize;
@@ -1012,7 +1012,7 @@ void UGameplayDebuggingComponent::PrepareNavMeshData(struct FNavMeshSceneProxyDa
 			UncompressedBuffer.AddZeroed(UncompressedSize);
 
 			FCompression::UncompressMemory((ECompressionFlags)(COMPRESS_ZLIB),
-				(void*)UncompressedBuffer.GetTypedData(), UncompressedSize, SrcBuffer, CompressedSize);
+				(void*)UncompressedBuffer.GetData(), UncompressedSize, SrcBuffer, CompressedSize);
 		}
 
 		// read serialized values
diff --git a/Engine/Source/Developer/HotReload/Private/HotReloadClassReinstancer.cpp b/Engine/Source/Developer/HotReload/Private/HotReloadClassReinstancer.cpp
--- a/Engine/Source/Developer/HotReload/Private/HotReloadClassReinstancer.cpp
+++ b/Engine/Source/Developer/HotReload/Private/HotReloadClassReinstancer.cpp
@@ -191,7 +191,7 @@ void FHotReloadClassReinstancer::RecreateCDOAndSetupOldClassReinstancing(UClass*
 
 	// We only want to re-instance the old class if its CDO's values have changed or any of its DSOs' property values have changed
 	if (OriginalCDOProperties.Num() != ReconstructedCDOProperties.Num() ||
-		FMemory::Memcmp(OriginalCDOProperties.GetTypedData(), ReconstructedCDOProperties.GetTypedData(), OriginalCDOProperties.Num()))
+		FMemory::Memcmp(OriginalCDOProperties.GetData(), ReconstructedCDOProperties.GetData(), OriginalCDOProperties.Num()))
 	{
 		bNeedsReinstancing = true;
 		SaveClassFieldMapping(InOldClass);
diff --git a/Engine/Source/Developer/ImageCore/Public/ImageCore.h b/Engine/Source/Developer/ImageCore/Public/ImageCore.h
--- a/Engine/Source/Developer/ImageCore/Public/ImageCore.h
+++ b/Engine/Source/Developer/ImageCore/Public/ImageCore.h
@@ -129,37 +129,37 @@ public:
 	uint8* AsG8()
 	{
 		check(Format == ERawImageFormat::G8);
-		return RawData.GetTypedData();
+		return RawData.GetData();
 	}
 
 	class FColor* AsBGRA8()
 	{
 		check(Format == ERawImageFormat::BGRA8);
-		return (class FColor*)RawData.GetTypedData();
+		return (class FColor*)RawData.GetData();
 	}
 
 	class FColor* AsBGRE8()
 	{
 		check(Format == ERawImageFormat::BGRE8);
-		return (class FColor*)RawData.GetTypedData();
+		return (class FColor*)RawData.GetData();
 	}
 
 	uint16* AsRGBA16()
 	{
 		check(Format == ERawImageFormat::RGBA16);
-		return (uint16*)RawData.GetTypedData();
+		return (uint16*)RawData.GetData();
 	}
 
 	class FFloat16Color* AsRGBA16F()
 	{
 		check(Format == ERawImageFormat::RGBA16F);
-		return (class FFloat16Color*)RawData.GetTypedData();
+		return (class FFloat16Color*)RawData.GetData();
 	}
 
 	struct FLinearColor* AsRGBA32F()
 	{
 		check(Format == ERawImageFormat::RGBA32F);
-		return (struct FLinearColor*)RawData.GetTypedData();
+		return (struct FLinearColor*)RawData.GetData();
 	}
 
 	// Convenience accessors to const raw data
@@ -167,36 +167,36 @@ public:
 	const uint8* AsG8() const
 	{
 		check(Format == ERawImageFormat::G8);
-		return RawData.GetTypedData();
+		return RawData.GetData();
 	}
 
 	const class FColor* AsBGRA8() const
 	{
 		check(Format == ERawImageFormat::BGRA8);
-		return (const class FColor*)RawData.GetTypedData();
+		return (const class FColor*)RawData.GetData();
 	}
 
 	const class FColor* AsBGRE8() const
 	{
 		check(Format == ERawImageFormat::BGRE8);
-		return (class FColor*)RawData.GetTypedData();
+		return (class FColor*)RawData.GetData();
 	}
 
 	const uint16* AsRGBA16() const
 	{
 		check(Format == ERawImageFormat::RGBA16);
-		return (const uint16*)RawData.GetTypedData();
+		return (const uint16*)RawData.GetData();
 	}
 
 	const class FFloat16Color* AsRGBA16F() const
 	{
 		check(Format == ERawImageFormat::RGBA16F);
-		return (const class FFloat16Color*)RawData.GetTypedData();
+		return (const class FFloat16Color*)RawData.GetData();
 	}
 
 	const struct FLinearColor* AsRGBA32F() const
 	{
 		check(Format == ERawImageFormat::RGBA32F);
-		return (struct FLinearColor*)RawData.GetTypedData();
+		return (struct FLinearColor*)RawData.GetData();
 	}
 };
diff --git a/Engine/Source/Developer/ImageWrapper/Private/ImageWrapperBase.cpp b/Engine/Source/Developer/ImageWrapper/Private/ImageWrapperBase.cpp
--- a/Engine/Source/Developer/ImageWrapper/Private/ImageWrapperBase.cpp
+++ b/Engine/Source/Developer/ImageWrapper/Private/ImageWrapperBase.cpp
@@ -74,7 +74,7 @@ bool FImageWrapperBase::SetCompressed( const void* InCompressedData, int32 InCom
 
 	CompressedData.Empty( InCompressedSize );
 	CompressedData.AddUninitialized( InCompressedSize );
-	FMemory::Memcpy( CompressedData.GetTypedData(), InCompressedData, InCompressedSize );
+	FMemory::Memcpy(CompressedData.GetData(), InCompressedData, InCompressedSize);
 
 	return true;
 }
@@ -92,7 +92,7 @@ bool FImageWrapperBase::SetRaw( const void* InRawData, int32 InRawSize, const in
 
 	RawData.Empty(InRawSize);
 	RawData.AddUninitialized(InRawSize);
-	FMemory::Memcpy(RawData.GetTypedData(), InRawData, InRawSize);
+	FMemory::Memcpy(RawData.GetData(), InRawData, InRawSize);
 
 	RawFormat = InFormat;
 	RawBitDepth = InBitDepth;
diff --git a/Engine/Source/Developer/ImageWrapper/Private/JpegImageWrapper.cpp b/Engine/Source/Developer/ImageWrapper/Private/JpegImageWrapper.cpp
--- a/Engine/Source/Developer/ImageWrapper/Private/JpegImageWrapper.cpp
+++ b/Engine/Source/Developer/ImageWrapper/Private/JpegImageWrapper.cpp
@@ -115,7 +115,7 @@ void FJpegImageWrapper::Compress( int32 Quality )
 		jpge::params Parameters;
 		Parameters.m_quality = Quality;
 		bool bSuccess = jpge::compress_image_to_jpeg_file_in_memory(
-			CompressedData.GetTypedData(), OutBufferSize, Width, Height, NumComponents, RawData.GetTypedData(), Parameters);
+			CompressedData.GetData(), OutBufferSize, Width, Height, NumComponents, RawData.GetData(), Parameters);
 		
 		check(bSuccess);
 
@@ -153,13 +153,13 @@ void FJpegImageWrapper::Uncompress( const ERGBFormat::Type InFormat, int32 InBit
 
 	int32 NumColors;
 	uint8* OutData = jpgd::decompress_jpeg_image_from_memory(
-		CompressedData.GetTypedData(), CompressedData.Num(), &Width, &Height, &NumColors, Channels );
+		CompressedData.GetData(), CompressedData.Num(), &Width, &Height, &NumColors, Channels);
 
 	RawData.Empty();
 	RawData.AddUninitialized( Width * Height * Channels );
 	if (OutData)
 	{
-		FMemory::Memcpy( RawData.GetTypedData(), OutData, RawData.Num() );
+		FMemory::Memcpy(RawData.GetData(), OutData, RawData.Num());
 		FMemory::Free(OutData);
 	}
 }
diff --git a/Engine/Source/Developer/LogVisualizer/Private/LogVisualizer.cpp b/Engine/Source/Developer/LogVisualizer/Private/LogVisualizer.cpp
--- a/Engine/Source/Developer/LogVisualizer/Private/LogVisualizer.cpp
+++ b/Engine/Source/Developer/LogVisualizer/Private/LogVisualizer.cpp
@@ -154,7 +154,7 @@ int32 FLogVisualizer::GetLogIndexForActor(const AActor* Actor)
 {
 	int32 ResultIndex = INDEX_NONE;
 	const FString FullName = Actor->GetFullName();
-	TSharedPtr<FActorsVisLog>* Log = Logs.GetTypedData();
+	TSharedPtr<FActorsVisLog>* Log = Logs.GetData();
 	for (int32 LogIndex = 0; LogIndex < Logs.Num(); ++LogIndex, ++Log)
 	{
 		if ((*Log).IsValid() && (*Log)->FullName == FullName)
diff --git a/Engine/Source/Developer/LogVisualizer/Private/SLogVisualizer.cpp b/Engine/Source/Developer/LogVisualizer/Private/SLogVisualizer.cpp
--- a/Engine/Source/Developer/LogVisualizer/Private/SLogVisualizer.cpp
+++ b/Engine/Source/Developer/LogVisualizer/Private/SLogVisualizer.cpp
@@ -543,7 +543,7 @@ void SLogVisualizer::Construct(const FArguments& InArgs, FLogVisualizer* InLogVi
 	}
 
 	TArray<TSharedPtr<FActorsVisLog> >& Logs = LogVisualizer->Logs;
-	TSharedPtr<FActorsVisLog>* SharedLog = Logs.GetTypedData();
+	TSharedPtr<FActorsVisLog>* SharedLog = Logs.GetData();
 	for (int32 LogIndex = 0; LogIndex < Logs.Num(); ++LogIndex, ++SharedLog)
 	{
 		if (SharedLog->IsValid())
@@ -1321,7 +1321,7 @@ void SLogVisualizer::AddOrUpdateLog(int32 LogIndex, const FActorsVisLog* Log)
 
 void SLogVisualizer::DoFullUpdate()
 {
-	TSharedPtr<FLogsListItem>* LogListItem = LogsList.GetTypedData();
+	TSharedPtr<FLogsListItem>* LogListItem = LogsList.GetData();
 	for (int32 ItemIndex = 0; ItemIndex < LogsList.Num(); ++ItemIndex, ++LogListItem)
 	{
 		if (LogListItem->IsValid() && LogVisualizer->Logs.IsValidIndex((*LogListItem)->LogIndex))
@@ -1362,7 +1362,7 @@ void SLogVisualizer::DoFullUpdate()
 
 void SLogVisualizer::DoTickUpdate()
 {
-	TSharedPtr<FLogsListItem>* LogListItem = LogsList.GetTypedData();
+	TSharedPtr<FLogsListItem>* LogListItem = LogsList.GetData();
 	for (int32 ItemIndex = 0; ItemIndex < LogsList.Num(); ++ItemIndex, ++LogListItem)
 	{
 		if (LogListItem->IsValid() && LogVisualizer->Logs.IsValidIndex((*LogListItem)->LogIndex))
@@ -1732,7 +1732,7 @@ void SLogVisualizer::ShowEntry(const FVisLogEntry* LogEntry)
 	UpdateStatusItems(LogEntry);
 	LogEntryLines.Reset();
 	
-	const FVisLogEntry::FLogLine* LogLine = LogEntry->LogLines.GetTypedData();
+	const FVisLogEntry::FLogLine* LogLine = LogEntry->LogLines.GetData();
 	for (int LineIndex = 0; LineIndex < LogEntry->LogLines.Num(); ++LineIndex, ++LogLine)
 	{
 		bool bShowLine = true;
@@ -1812,7 +1812,7 @@ void SLogVisualizer::RebuildFilteredList()
 	// redo selection
 	if (ItemsToSelect.Num() > 0)
 	{
-		TSharedPtr<FLogsListItem>* Item = ItemsToSelect.GetTypedData();
+		TSharedPtr<FLogsListItem>* Item = ItemsToSelect.GetData();
 		for (int32 ItemIndex = 0; ItemIndex < ItemsToSelect.Num(); ++ItemIndex, ++Item)
 		{
 			const int32 IndexInList = FindIndexInLogsList((*Item)->LogIndex);
@@ -1997,7 +1997,7 @@ void SLogVisualizer::DrawOnCanvas(UCanvas* Canvas, APlayerController*)
 
 		if (bDrawLogEntriesPath)
 		{
-			const TSharedPtr<FVisLogEntry>* Entry = Entries.GetTypedData();
+			const TSharedPtr<FVisLogEntry>* Entry = Entries.GetData();
 			FVector Location = (*Entry)->Location;
 			++Entry;
 
@@ -2226,7 +2226,7 @@ void SLogVisualizer::DrawOnCanvas(UCanvas* Canvas, APlayerController*)
 				}
 			}
 
-			const FVisLogEntry::FElementToDraw* ElementToDraw = Entry->ElementsToDraw.GetTypedData();
+			const FVisLogEntry::FElementToDraw* ElementToDraw = Entry->ElementsToDraw.GetData();
 			const int32 ElementsCount = Entry->ElementsToDraw.Num();
 			
 			for (int32 ElementIndex = 0; ElementIndex < ElementsCount; ++ElementIndex, ++ElementToDraw)
@@ -2245,7 +2245,7 @@ void SLogVisualizer::DrawOnCanvas(UCanvas* Canvas, APlayerController*)
 					{
 						const float Radius = float(ElementToDraw->Radius);
 						const bool bDrawLabel = ElementToDraw->Description.IsEmpty() == false;
-						const FVector* Location = ElementToDraw->Points.GetTypedData();
+						const FVector* Location = ElementToDraw->Points.GetData();
 						for (int32 Index = 0; Index < ElementToDraw->Points.Num(); ++Index, ++Location)
 						{
 							DrawDebugSphere(World, *Location, Radius, 16, Color);
@@ -2261,7 +2261,7 @@ void SLogVisualizer::DrawOnCanvas(UCanvas* Canvas, APlayerController*)
 					{
 						const float Thickness = float(ElementToDraw->Thicknes);
 						const bool bDrawLabel = ElementToDraw->Description.IsEmpty() == false && ElementToDraw->Points.Num() > 2;
-						const FVector* Location = ElementToDraw->Points.GetTypedData();
+						const FVector* Location = ElementToDraw->Points.GetData();
 						for (int32 Index = 0; Index + 1 < ElementToDraw->Points.Num(); Index += 2, Location += 2)
 						{
 							DrawDebugLine(World, *Location, *(Location + 1), Color
@@ -2305,7 +2305,7 @@ void SLogVisualizer::DrawOnCanvas(UCanvas* Canvas, APlayerController*)
 					{
 						const float Thickness = float(ElementToDraw->Thicknes);
 						const bool bDrawLabel = ElementToDraw->Description.IsEmpty() == false && ElementToDraw->Points.Num() > 2;
-						const FVector* BoxExtent = ElementToDraw->Points.GetTypedData();
+						const FVector* BoxExtent = ElementToDraw->Points.GetData();
 						for (int32 Index = 0; Index + 1 < ElementToDraw->Points.Num(); Index += 2, BoxExtent += 2)
 						{
 							FBox Box(*BoxExtent, *(BoxExtent + 1));
@@ -2710,7 +2710,7 @@ void SLogVisualizer::SaveSelectedLogs(FString& Filename)
 	EntriesArray.Reserve(ItemsToSave.Num());
 
 
-	TSharedPtr<FLogsListItem>* LogListItem = ItemsToSave.GetTypedData();
+	TSharedPtr<FLogsListItem>* LogListItem = ItemsToSave.GetData();
 	for (int32 ItemIndex = 0; ItemIndex < ItemsToSave.Num(); ++ItemIndex, ++LogListItem)
 	{
 		if (LogListItem->IsValid() && LogVisualizer->Logs.IsValidIndex((*LogListItem)->LogIndex))
diff --git a/Engine/Source/Developer/MeshBoneReduction/Private/MeshBoneReduction.cpp b/Engine/Source/Developer/MeshBoneReduction/Private/MeshBoneReduction.cpp
--- a/Engine/Source/Developer/MeshBoneReduction/Private/MeshBoneReduction.cpp
+++ b/Engine/Source/Developer/MeshBoneReduction/Private/MeshBoneReduction.cpp
@@ -235,10 +235,10 @@ public:
 		// Insert a new LOD model entry if needed.
 		if ( DesiredLOD == SkeletalMeshResource->LODModels.Num() )
 		{
-			SkeletalMeshResource->LODModels.AddRawItem(0);
+			SkeletalMeshResource->LODModels.Add(0);
 		}
 
-		FStaticLODModel** LODModels = SkeletalMeshResource->LODModels.GetTypedData();
+		FStaticLODModel** LODModels = SkeletalMeshResource->LODModels.GetData();
 		delete LODModels[DesiredLOD];
 		FStaticLODModel* NewModel = new FStaticLODModel();
 		LODModels[DesiredLOD] = NewModel;
diff --git a/Engine/Source/Developer/MeshUtilities/Private/MeshUtilities.cpp b/Engine/Source/Developer/MeshUtilities/Private/MeshUtilities.cpp
--- a/Engine/Source/Developer/MeshUtilities/Private/MeshUtilities.cpp
+++ b/Engine/Source/Developer/MeshUtilities/Private/MeshUtilities.cpp
@@ -581,14 +581,14 @@ namespace NvTriStrip
 		SetListsOnly(true);
 		SetCacheSize(CACHESIZE_GEFORCE3);
 
-		GenerateStrips((uint8*)Indices.GetTypedData(),Is32Bit,Indices.Num(),&PrimitiveGroups,&NumPrimitiveGroups);
+		GenerateStrips((uint8*)Indices.GetData(),Is32Bit,Indices.Num(),&PrimitiveGroups,&NumPrimitiveGroups);
 
 		Indices.Empty();
 		Indices.AddUninitialized(PrimitiveGroups->numIndices);
 	
 		if( Is32Bit )
 		{
-			FMemory::Memcpy(Indices.GetTypedData(),PrimitiveGroups->indices,Indices.Num() * sizeof(IndexDataType));
+			FMemory::Memcpy(Indices.GetData(),PrimitiveGroups->indices,Indices.Num() * sizeof(IndexDataType));
 		}
 		else
 		{
@@ -666,7 +666,7 @@ namespace Forsyth
 
 		if( Is32Bit )
 		{
-			FMemory::Memcpy(Indices.GetTypedData(),OptimizedIndices.GetTypedData(),Indices.Num() * sizeof(IndexDataType));
+			FMemory::Memcpy(Indices.GetData(),OptimizedIndices.GetData(),Indices.Num() * sizeof(IndexDataType));
 		}
 		else
 		{
@@ -722,7 +722,7 @@ public:
 		, VertexBuffer( InVertexBuffer )
 	{
 		check( PositionVertexBuffer.GetNumVertices() == VertexBuffer.GetNumVertices() );
-		mIb = new nv::IndexBuffer( (void*)Indices.GetTypedData(), nv::IBT_U32, Indices.Num(), false );
+		mIb = new nv::IndexBuffer( (void*)Indices.GetData(), nv::IBT_U32, Indices.Num(), false );
 	}
 
 	/** Retrieve the position and first texture coordinate of the specified index. */
@@ -780,7 +780,7 @@ public:
 		: VertexBuffer( InVertexBuffer )
 		, TexCoordCount( InTexCoordCount )
 	{
-		mIb = new nv::IndexBuffer( (void*)Indices.GetTypedData(), nv::IBT_U32, Indices.Num(), false );
+		mIb = new nv::IndexBuffer( (void*)Indices.GetData(), nv::IBT_U32, Indices.Num(), false );
 	}
 
 	/** Retrieve the position and first texture coordinate of the specified index. */
@@ -2425,7 +2425,7 @@ bool FMeshUtilities::BuildStaticMesh(
 
 				CombinedIndices.AddUninitialized(SectionIndices.Num());
 				uint32* DestPtr = &CombinedIndices[Section.FirstIndex];
-				uint32 const* SrcPtr = SectionIndices.GetTypedData();
+				uint32 const* SrcPtr = SectionIndices.GetData();
 
 				Section.MinVertexIndex = *SrcPtr;
 				Section.MaxVertexIndex = *SrcPtr;
diff --git a/Engine/Source/Developer/Profiler/Private/ProfilerStream.cpp b/Engine/Source/Developer/Profiler/Private/ProfilerStream.cpp
--- a/Engine/Source/Developer/Profiler/Private/ProfilerStream.cpp
+++ b/Engine/Source/Developer/Profiler/Private/ProfilerStream.cpp
@@ -36,7 +36,7 @@ void FProfilerUIStream::GenerateUIStream( const FProfilerStream& ProfilerStream,
 				FProfilerUIStackNode* ThreadUIStackNode = new FProfilerUIStackNode( ProfilerFrame->Root->Children[ThreadIndex], FProfilerUIStackNode::THREAD_NODE_INDEX, FProfilerUIStackNode::THREAD_NODE_INDEX, FrameIndex );
 				ThreadUIStackNode->InitializeUIData( NumMillisecondsPerWindow, NumPixelsPerMillisecond, NumMillisecondsPerSample );
 
-				ThreadNodes.AddRawItem( ThreadUIStackNode );
+				ThreadNodes.Add( ThreadUIStackNode );
 
 				CombineOrSet( ThreadUIStackNode, *ProfilerFrame->Root->Children[ThreadIndex], ThreadIndex*DEFAULT_VISIBLE_THREAD_DEPTH, NumMillisecondsPerWindow, NumPixelsPerMillisecond, NumMillisecondsPerSample );
 			}
@@ -80,7 +80,7 @@ void FProfilerUIStream::CombineOrSet( FProfilerUIStackNode* ParentUIStackNode, c
 					{
 						FProfilerUIStackNode* ProfilerUIStackNode = new FProfilerUIStackNode( OriginalStackNodes, OriginalStackNodes.Num(), GlobalNodeDepth, ThreadIndex, FrameIndex );
 						ProfilerUIStackNode->InitializeUIData( NumMillisecondsPerWindow, NumPixelsPerMillisecond, NumMillisecondsPerSample );
-						ParentUIStackNode->Children.AddRawItem( ProfilerUIStackNode );
+						ParentUIStackNode->Children.Add( ProfilerUIStackNode );
 						bLookingForChildToCombine = false;
 						OriginalStackNodes.Reset();
 					}
@@ -89,7 +89,7 @@ void FProfilerUIStream::CombineOrSet( FProfilerUIStackNode* ParentUIStackNode, c
 					{
 						FProfilerUIStackNode* ProfilerUIStackNode = new FProfilerUIStackNode( NextProfilerChildNode, GlobalNodeDepth, ThreadIndex, FrameIndex );
 						ProfilerUIStackNode->InitializeUIData( NumMillisecondsPerWindow, NumPixelsPerMillisecond, NumMillisecondsPerSample );
-						ParentUIStackNode->Children.AddRawItem( ProfilerUIStackNode );
+						ParentUIStackNode->Children.Add( ProfilerUIStackNode );
 
 						// Check the call stack depth.
 						if( ThreadNodeDepth != DEFAULT_VISIBLE_THREAD_DEPTH - 1 )
@@ -118,7 +118,7 @@ void FProfilerUIStream::CombineOrSet( FProfilerUIStackNode* ParentUIStackNode, c
 					// Combine child nodes into one node.
 					FProfilerUIStackNode* ProfilerUIStackNode = new FProfilerUIStackNode( OriginalStackNodes, OriginalStackNodes.Num(), GlobalNodeDepth, ThreadIndex, FrameIndex );
 					ProfilerUIStackNode->InitializeUIData( NumMillisecondsPerWindow, NumPixelsPerMillisecond, NumMillisecondsPerSample );
-					ParentUIStackNode->Children.AddRawItem( ProfilerUIStackNode );
+					ParentUIStackNode->Children.Add( ProfilerUIStackNode );
 					bLookingForChildToCombine = false;
 					OriginalStackNodes.Reset();
 				}
@@ -134,7 +134,7 @@ void FProfilerUIStream::CombineOrSet( FProfilerUIStackNode* ParentUIStackNode, c
 				// We have a sample that can be displayed in the UI, add it.
 				FProfilerUIStackNode* ProfilerUIStackNode = new FProfilerUIStackNode( FirstProfilerChildNode, GlobalNodeDepth, ThreadIndex, FrameIndex );
 				ProfilerUIStackNode->InitializeUIData( NumMillisecondsPerWindow, NumPixelsPerMillisecond, NumMillisecondsPerSample );
-				ParentUIStackNode->Children.AddRawItem( ProfilerUIStackNode );
+				ParentUIStackNode->Children.Add( ProfilerUIStackNode );
 
 				// Check the call stack depth.
 				if( ThreadNodeDepth != DEFAULT_VISIBLE_THREAD_DEPTH - 1 )
@@ -161,7 +161,7 @@ void FProfilerUIStream::CombineOrSet( FProfilerUIStackNode* ParentUIStackNode, c
 			{
 				FProfilerUIStackNode* ProfilerUIStackNode = new FProfilerUIStackNode( OriginalStackNodes, OriginalStackNodes.Num(), GlobalNodeDepth, ThreadIndex, FrameIndex );
 				ProfilerUIStackNode->InitializeUIData( NumMillisecondsPerWindow, NumPixelsPerMillisecond, NumMillisecondsPerSample );
-				ParentUIStackNode->Children.AddRawItem( ProfilerUIStackNode );
+				ParentUIStackNode->Children.Add( ProfilerUIStackNode );
 				bLookingForChildToCombine = false;
 				OriginalStackNodes.Reset();
 			}
diff --git a/Engine/Source/Developer/RealtimeProfiler/Private/SRealtimeProfilerLineGraph.cpp b/Engine/Source/Developer/RealtimeProfiler/Private/SRealtimeProfilerLineGraph.cpp
--- a/Engine/Source/Developer/RealtimeProfiler/Private/SRealtimeProfilerLineGraph.cpp
+++ b/Engine/Source/Developer/RealtimeProfiler/Private/SRealtimeProfilerLineGraph.cpp
@@ -183,25 +183,25 @@ int32 SRealtimeProfilerLineGraph::OnPaint( const FPaintArgs& Args, const FGeomet
 
 		if(!bDisplayFPSChart)
 		{
-			YPos = ProfileDataArray.GetTypedData()[i]->DurationMs/MaxValue.Get();
+			YPos = ProfileDataArray.GetData()[i]->DurationMs/MaxValue.Get();
 			YPos = (AllottedGeometry.Size.Y-1) - (YPos*AllottedGeometry.Size.Y);
 			LinePoints.Add(FVector2D((int32)XPos,(int32)YPos));
 		}
 		else
 		{
-			YPos = FPSChartDataArray.GetTypedData()[i].UnitFrame/MaxValue.Get();
+			YPos = FPSChartDataArray.GetData()[i].UnitFrame/MaxValue.Get();
 			YPos = (AllottedGeometry.Size.Y-1) - (YPos*AllottedGeometry.Size.Y);
 			UnitFramePoints.Add(FVector2D((int32)XPos,(int32)YPos));
 
-			YPos = FPSChartDataArray.GetTypedData()[i].UnitRender/MaxValue.Get();
+			YPos = FPSChartDataArray.GetData()[i].UnitRender/MaxValue.Get();
 			YPos = (AllottedGeometry.Size.Y-1) - (YPos*AllottedGeometry.Size.Y);
 			UnitRenderPoints.Add(FVector2D((int32)XPos,(int32)YPos));
 
-			YPos = FPSChartDataArray.GetTypedData()[i].UnitGame/MaxValue.Get();
+			YPos = FPSChartDataArray.GetData()[i].UnitGame/MaxValue.Get();
 			YPos = (AllottedGeometry.Size.Y-1) - (YPos*AllottedGeometry.Size.Y);
 			UnitGamePoints.Add(FVector2D((int32)XPos,(int32)YPos));
 
-			YPos = FPSChartDataArray.GetTypedData()[i].UnitGPU/MaxValue.Get();
+			YPos = FPSChartDataArray.GetData()[i].UnitGPU/MaxValue.Get();
 			YPos = (AllottedGeometry.Size.Y-1) - (YPos*AllottedGeometry.Size.Y);
 			UnitGPUPoints.Add(FVector2D((int32)XPos,(int32)YPos));
 		}
@@ -367,7 +367,7 @@ void SRealtimeProfilerLineGraph::DisplayFrameDetailAtMouse(const FGeometry& InMy
 
 	if(0 <= FrameIndex && FrameIndex < ProfileDataArray.Num())
 	{
-		TSharedPtr< FVisualizerEvent > SelectedData = ProfileDataArray.GetTypedData()[FrameIndex];
+		TSharedPtr< FVisualizerEvent > SelectedData = ProfileDataArray.GetData()[FrameIndex];
 		Visualizer.Get()->DisplayFrameDetails(SelectedData);
 	}
 }
diff --git a/Engine/Source/Developer/ShaderFormatOpenGL/Private/OpenGLShaderCompiler.cpp b/Engine/Source/Developer/ShaderFormatOpenGL/Private/OpenGLShaderCompiler.cpp
--- a/Engine/Source/Developer/ShaderFormatOpenGL/Private/OpenGLShaderCompiler.cpp
+++ b/Engine/Source/Developer/ShaderFormatOpenGL/Private/OpenGLShaderCompiler.cpp
@@ -1361,8 +1361,8 @@ static void PrecompileShader(FShaderCompilerOutput& ShaderOutput, const FShaderC
 
 					RawCompileLog.Empty(LogLength);
 					RawCompileLog.AddZeroed(LogLength);
-					glGetShaderInfoLog(Shader, LogLength, /*OutLength=*/ NULL, RawCompileLog.GetTypedData());
-					CompileLog = ANSI_TO_TCHAR(RawCompileLog.GetTypedData());
+					glGetShaderInfoLog(Shader, LogLength, /*OutLength=*/ NULL, RawCompileLog.GetData());
+					CompileLog = ANSI_TO_TCHAR(RawCompileLog.GetData());
 					CompileLog.ParseIntoArray(&LogLines, TEXT("\n"), true);
 
 					for (int32 Line = 0; Line < LogLines.Num(); ++Line)
@@ -1376,7 +1376,7 @@ static void PrecompileShader(FShaderCompilerOutput& ShaderOutput, const FShaderC
 						NewError->StrippedErrorMessage = FString::Printf(
 							TEXT("GLSL source:\n%sGL compile log: %s\n"),
 							ANSI_TO_TCHAR(ShaderSource),
-							ANSI_TO_TCHAR(RawCompileLog.GetTypedData())
+							ANSI_TO_TCHAR(RawCompileLog.GetData())
 							);
 					}
 				}
diff --git a/Engine/Source/Developer/ShaderPreprocessor/Private/ShaderPreprocessor.cpp b/Engine/Source/Developer/ShaderPreprocessor/Private/ShaderPreprocessor.cpp
--- a/Engine/Source/Developer/ShaderPreprocessor/Private/ShaderPreprocessor.cpp
+++ b/Engine/Source/Developer/ShaderPreprocessor/Private/ShaderPreprocessor.cpp
@@ -164,7 +164,7 @@ private:
 
 		if (OutContents)
 		{
-			*OutContents = CachedContents ? CachedContents->GetTypedData() : NULL;
+			*OutContents = CachedContents ? CachedContents->GetData() : NULL;
 		}
 		if (OutContentSize)
 		{
diff --git a/Engine/Source/Developer/SimplygonMeshReduction/Private/SimplygonMeshReduction.cpp b/Engine/Source/Developer/SimplygonMeshReduction/Private/SimplygonMeshReduction.cpp
--- a/Engine/Source/Developer/SimplygonMeshReduction/Private/SimplygonMeshReduction.cpp
+++ b/Engine/Source/Developer/SimplygonMeshReduction/Private/SimplygonMeshReduction.cpp
@@ -225,7 +225,7 @@ public:
 		// Insert a new LOD model entry if needed.
 		if ( LODIndex == SkeletalMeshResource->LODModels.Num() )
 		{
-			SkeletalMeshResource->LODModels.AddRawItem(0);
+			SkeletalMeshResource->LODModels.Add(0);
 		}
 
 		// We'll need to store the max deviation after optimization if we wish to recalculate the LOD's display distance
@@ -233,7 +233,7 @@ public:
 
 		// Swap in a new model, delete the old.
 		check( LODIndex < SkeletalMeshResource->LODModels.Num() );
-		FStaticLODModel** LODModels = SkeletalMeshResource->LODModels.GetTypedData();
+		FStaticLODModel** LODModels = SkeletalMeshResource->LODModels.GetData();
 		delete LODModels[LODIndex];
 
 		// Copy over LOD info from LOD0 if there is no previous info.
@@ -404,7 +404,7 @@ public:
 		TArray<uint8> LicenseFileContents;
 		if (FFileHelper::LoadFileToArray(LicenseFileContents, *FPaths::Combine(*DllPath, TEXT("license.dat")), FILEREAD_Silent) && LicenseFileContents.Num() > 0)
 		{
-			LicenseData = (const char*)LicenseFileContents.GetTypedData();
+			LicenseData = (const char*)LicenseFileContents.GetData();
 		}
 
 		SimplygonSDK::ISimplygonSDK* SDK = NULL;
diff --git a/Engine/Source/Developer/StandaloneRenderer/Private/Windows/D3D/SlateD3DTextureManager.cpp b/Engine/Source/Developer/StandaloneRenderer/Private/Windows/D3D/SlateD3DTextureManager.cpp
--- a/Engine/Source/Developer/StandaloneRenderer/Private/Windows/D3D/SlateD3DTextureManager.cpp
+++ b/Engine/Source/Developer/StandaloneRenderer/Private/Windows/D3D/SlateD3DTextureManager.cpp
@@ -296,7 +296,7 @@ FSlateShaderResourceProxy* FSlateD3DTextureManager::CreateDynamicTextureResource
 	TextureResource = MakeShareable(new FDynamicTextureResource(LoadedTexture));
 
 	D3D11_SUBRESOURCE_DATA InitData;
-	InitData.pSysMem = RawData.GetTypedData();
+	InitData.pSysMem = RawData.GetData();
 	InitData.SysMemPitch = Width * 4;
 
 	FNewTextureInfo Info;
@@ -392,7 +392,7 @@ FSlateShaderResourceProxy* FSlateD3DTextureManager::GenerateTextureResource( con
 		NewProxy->ActualSize = FIntPoint( Width, Height );
 
 		D3D11_SUBRESOURCE_DATA InitData;
-		InitData.pSysMem = Info.TextureData->GetRawBytes().GetTypedData();
+		InitData.pSysMem = Info.TextureData->GetRawBytes().GetData();
 		InitData.SysMemPitch = Width * 4;
 
 		Texture->Init( Info.bSrgb ? DXGI_FORMAT_R8G8B8A8_UNORM_SRGB : DXGI_FORMAT_R8G8B8A8_UNORM, &InitData );
diff --git a/Engine/Source/Developer/StandaloneRenderer/Private/Windows/D3D/SlateD3DTextures.cpp b/Engine/Source/Developer/StandaloneRenderer/Private/Windows/D3D/SlateD3DTextures.cpp
--- a/Engine/Source/Developer/StandaloneRenderer/Private/Windows/D3D/SlateD3DTextures.cpp
+++ b/Engine/Source/Developer/StandaloneRenderer/Private/Windows/D3D/SlateD3DTextures.cpp
@@ -50,7 +50,7 @@ void FSlateD3DTexture::UpdateTexture(const TArray<uint8>& Bytes)
 	GD3DDeviceContext->Map(D3DTexture, 0, D3D11_MAP_WRITE_DISCARD, 0, &Resource);
 	for (uint32 Row = 0; Row < SizeY; ++Row)
 	{
-		FMemory::Memcpy((uint8*)Resource.pData + Row * Resource.RowPitch, Bytes.GetTypedData() + Row * SizeX*4, SizeX*4);
+		FMemory::Memcpy((uint8*)Resource.pData + Row * Resource.RowPitch, Bytes.GetData() + Row * SizeX*4, SizeX*4);
 	}
 	GD3DDeviceContext->Unmap(D3DTexture, 0);
 }
@@ -105,7 +105,7 @@ void FSlateFontAtlasD3D::ConditionalUpdateTexture()
 	{
 		D3D11_MAPPED_SUBRESOURCE Resource;
 		GD3DDeviceContext->Map( FontTexture->GetTextureResource(), 0, D3D11_MAP_WRITE_DISCARD, 0, &Resource );
-		FMemory::Memcpy( Resource.pData, AtlasData.GetTypedData(), sizeof(uint8)*AtlasWidth*AtlasHeight);
+		FMemory::Memcpy( Resource.pData, AtlasData.GetData(), sizeof(uint8)*AtlasWidth*AtlasHeight);
 		GD3DDeviceContext->Unmap( FontTexture->GetTextureResource(), 0 );
 
 		bNeedsUpdate = false;
diff --git a/Engine/Source/Developer/TextureFormatAndroid/Private/TextureFormatAndroid.cpp b/Engine/Source/Developer/TextureFormatAndroid/Private/TextureFormatAndroid.cpp
--- a/Engine/Source/Developer/TextureFormatAndroid/Private/TextureFormatAndroid.cpp
+++ b/Engine/Source/Developer/TextureFormatAndroid/Private/TextureFormatAndroid.cpp
@@ -82,7 +82,7 @@ static bool CompressImageUsingQonvert(
 	DstImg.nWidth    = SizeX;
 	DstImg.nHeight   = SizeY;
 	DstImg.nDataSize = ImageBlocksX * ImageBlocksY * BlockBytes;
-	DstImg.pData     = OutCompressedData.GetTypedData();
+	DstImg.pData     = OutCompressedData.GetData();
 
 	switch (PixelFormat)
 	{
diff --git a/Engine/Source/Developer/TextureFormatDXT/Private/TextureFormatDXT.cpp b/Engine/Source/Developer/TextureFormatDXT/Private/TextureFormatDXT.cpp
--- a/Engine/Source/Developer/TextureFormatDXT/Private/TextureFormatDXT.cpp
+++ b/Engine/Source/Developer/TextureFormatDXT/Private/TextureFormatDXT.cpp
@@ -320,7 +320,7 @@ static bool CompressImageUsingNVTT(
 				SizeY,
 				bSRGB,
 				bIsNormalMap,
-				OutCompressedData.GetTypedData(),
+				OutCompressedData.GetData(),
 				OutCompressedData.Num()
 				);
 		}
@@ -342,7 +342,7 @@ static bool CompressImageUsingNVTT(
 	{
 		FScopeLock ScopeLock(&GNVCompressionCriticalSection);
 		const uint8* Src = (const uint8*)SourceData;
-		uint8* Dest = OutCompressedData.GetTypedData();
+		uint8* Dest = OutCompressedData.GetData();
 		for (int32 BatchIndex = 0; BatchIndex < NumBatches; ++BatchIndex)
 		{
 			new(Compressors) FNVTTCompressor(
diff --git a/Engine/Source/Developer/TextureFormatPVR/Private/TextureFormatPVR.cpp b/Engine/Source/Developer/TextureFormatPVR/Private/TextureFormatPVR.cpp
--- a/Engine/Source/Developer/TextureFormatPVR/Private/TextureFormatPVR.cpp
+++ b/Engine/Source/Developer/TextureFormatPVR/Private/TextureFormatPVR.cpp
@@ -530,7 +530,7 @@ class FTextureFormatPVR : public ITextureFormat
 			}
 
 			// Copy compressed data
-			FMemory::Memcpy(MipData, PVRData.GetTypedData() + FileOffset, DestNumBytes);
+			FMemory::Memcpy(MipData, PVRData.GetData() + FileOffset, DestNumBytes);
 		}
 
 		// Delete intermediate files
diff --git a/Engine/Source/Developer/TextureFormatUncompressed/Private/TextureFormatUncompressed.cpp b/Engine/Source/Developer/TextureFormatUncompressed/Private/TextureFormatUncompressed.cpp
--- a/Engine/Source/Developer/TextureFormatUncompressed/Private/TextureFormatUncompressed.cpp
+++ b/Engine/Source/Developer/TextureFormatUncompressed/Private/TextureFormatUncompressed.cpp
@@ -95,7 +95,7 @@ class FTextureFormatUncompressed : public ITextureFormat
 			OutCompressedImage.RawData.AddUninitialized(NumTexels * 2);
 			const FColor* FirstColor = Image.AsBGRA8();
 			const FColor* LastColor = FirstColor + NumTexels;
-			int8* Dest = (int8*)OutCompressedImage.RawData.GetTypedData();
+			int8* Dest = (int8*)OutCompressedImage.RawData.GetData();
 
 			for (const FColor* Color = FirstColor; Color < LastColor; ++Color)
 			{
@@ -132,7 +132,7 @@ class FTextureFormatUncompressed : public ITextureFormat
 			OutCompressedImage.RawData.AddUninitialized(NumTexels * 4);
 			const FColor* FirstColor = Image.AsBGRA8();
 			const FColor* LastColor = FirstColor + NumTexels;
-			int8* Dest = (int8*)OutCompressedImage.RawData.GetTypedData();
+			int8* Dest = (int8*)OutCompressedImage.RawData.GetData();
 
 			for (const FColor* Color = FirstColor; Color < LastColor; ++Color)
 			{
@@ -159,7 +159,7 @@ class FTextureFormatUncompressed : public ITextureFormat
 			OutCompressedImage.RawData.AddUninitialized(NumTexels * 4);
 			const FColor* FirstColor = Image.AsBGRA8();
 			const FColor* LastColor = FirstColor + NumTexels;
-			int8* Dest = (int8*)OutCompressedImage.RawData.GetTypedData();
+			int8* Dest = (int8*)OutCompressedImage.RawData.GetData();
 
 			for (const FColor* Color = FirstColor; Color < LastColor; ++Color)
 			{
diff --git a/Engine/Source/Editor/BehaviorTreeEditor/Private/SBehaviorTreeDiff.cpp b/Engine/Source/Editor/BehaviorTreeEditor/Private/SBehaviorTreeDiff.cpp
--- a/Engine/Source/Editor/BehaviorTreeEditor/Private/SBehaviorTreeDiff.cpp
+++ b/Engine/Source/Editor/BehaviorTreeEditor/Private/SBehaviorTreeDiff.cpp
@@ -186,7 +186,7 @@ TSharedRef<SWidget> SBehaviorTreeDiff::GenerateDiffListWidget()
 				return A->Result.Diff < B->Result.Diff;
 			}
 		};
-		Sort(DiffListSource.GetTypedData(),DiffListSource.Num(), FSortDiff());
+		Sort(DiffListSource.GetData(),DiffListSource.Num(), FSortDiff());
 
 		// Map commands through UI
 		const FDiffListCommands& Commands = FDiffListCommands::Get();
diff --git a/Engine/Source/Editor/FoliageEdit/Private/FoliageEdMode.cpp b/Engine/Source/Editor/FoliageEdit/Private/FoliageEdMode.cpp
--- a/Engine/Source/Editor/FoliageEdit/Private/FoliageEdMode.cpp
+++ b/Engine/Source/Editor/FoliageEdit/Private/FoliageEdMode.cpp
@@ -1325,7 +1325,7 @@ void FEdModeFoliage::ApplyPaintBucket(AActor* Actor, bool bRemove)
 				const FStaticMeshComponentLODInfo* InstanceMeshLODInfo = nullptr;
 				if (StaticMeshComponent->LODData.Num() > 0)
 				{
-					InstanceMeshLODInfo = StaticMeshComponent->LODData.GetTypedData();
+					InstanceMeshLODInfo = StaticMeshComponent->LODData.GetData();
 					bHasInstancedColorData = InstanceMeshLODInfo->PaintedVertices.Num() == LODModel.VertexBuffer.GetNumVertices();
 				}
 
@@ -1474,7 +1474,7 @@ bool FEdModeFoliage::GetStaticMeshVertexColorForHit(UStaticMeshComponent* InStat
 	const FStaticMeshComponentLODInfo* InstanceMeshLODInfo = nullptr;
 	if (InStaticMeshComponent->LODData.Num() > 0)
 	{
-		InstanceMeshLODInfo = InStaticMeshComponent->LODData.GetTypedData();
+		InstanceMeshLODInfo = InStaticMeshComponent->LODData.GetData();
 		bHasInstancedColorData = InstanceMeshLODInfo->PaintedVertices.Num() == LODModel.VertexBuffer.GetNumVertices();
 	}
 
diff --git a/Engine/Source/Editor/FontEditor/Private/FontEditor.cpp b/Engine/Source/Editor/FontEditor/Private/FontEditor.cpp
--- a/Engine/Source/Editor/FontEditor/Private/FontEditor.cpp
+++ b/Engine/Source/Editor/FontEditor/Private/FontEditor.cpp
@@ -732,7 +732,7 @@ bool FFontEditor::ImportPage(int32 PageNum, const TCHAR* FileName)
 	if (FFileHelper::LoadFileToArray(Data, FileName))
 	{
 		// Make a const pointer for the API to be happy
-		const uint8* DataPtr = Data.GetTypedData();
+		const uint8* DataPtr = Data.GetData();
 		
 		// Create the new texture... note RF_Public because font textures can be referenced directly by material expressions
 		UTexture2D* NewPage = (UTexture2D*)Factory->FactoryCreateBinary(UTexture2D::StaticClass(), Font, NAME_None, RF_Public, NULL, TEXT("TGA"), DataPtr, DataPtr + Data.Num(), GWarn);
diff --git a/Engine/Source/Editor/GeometryMode/Private/GeometryModifiers.cpp b/Engine/Source/Editor/GeometryMode/Private/GeometryModifiers.cpp
--- a/Engine/Source/Editor/GeometryMode/Private/GeometryModifiers.cpp
+++ b/Engine/Source/Editor/GeometryMode/Private/GeometryModifiers.cpp
@@ -1227,7 +1227,7 @@ void UGeomModifier_Pen::Apply()
 		}
 
 		// Generate center location from the shape's center
-		FBox WorldBounds(ShapeVertices.GetTypedData(), ShapeVertices.Num());
+		FBox WorldBounds(ShapeVertices.GetData(), ShapeVertices.Num());
 		FVector BaseLocation = WorldBounds.GetCenter();
 
 		//create a scoped transaction so that we can undo the creation/modification		
diff --git a/Engine/Source/Editor/Kismet/Private/SBlueprintDiff.cpp b/Engine/Source/Editor/Kismet/Private/SBlueprintDiff.cpp
--- a/Engine/Source/Editor/Kismet/Private/SBlueprintDiff.cpp
+++ b/Engine/Source/Editor/Kismet/Private/SBlueprintDiff.cpp
@@ -425,7 +425,7 @@ TSharedRef<SWidget> FListItemGraphToDiff::GenerateDiffListWidget()
 			}
 		};
 
-		Sort(DiffListSource.GetTypedData(),DiffListSource.Num(), SortDiff());
+		Sort(DiffListSource.GetData(),DiffListSource.Num(), SortDiff());
 
 		TSharedPtr<SListViewType> DiffListRef;
 		TSharedRef<SHorizontalBox> Result =	SNew(SHorizontalBox)
diff --git a/Engine/Source/Editor/LandscapeEditor/Private/LandscapeEdMode.cpp b/Engine/Source/Editor/LandscapeEditor/Private/LandscapeEdMode.cpp
--- a/Engine/Source/Editor/LandscapeEditor/Private/LandscapeEdMode.cpp
+++ b/Engine/Source/Editor/LandscapeEditor/Private/LandscapeEdMode.cpp
@@ -93,7 +93,7 @@ void ALandscape::SplitHeightmap(ULandscapeComponent* Comp, bool bMoveToCurrentLe
 		// Because of edge problem, normal would be just copy from old component data
 		TArray<uint8> NormalData;
 		NormalData.AddZeroed((1 + Comp->ComponentSizeQuads)*(1 + Comp->ComponentSizeQuads)*sizeof(uint16));
-		LandscapeEdit.GetHeightDataFast(Comp->GetSectionBase().X, Comp->GetSectionBase().Y, Comp->GetSectionBase().X + Comp->ComponentSizeQuads, Comp->GetSectionBase().Y + Comp->ComponentSizeQuads, (uint16*)HeightData.GetTypedData(), 0, (uint16*)NormalData.GetTypedData());
+		LandscapeEdit.GetHeightDataFast(Comp->GetSectionBase().X, Comp->GetSectionBase().Y, Comp->GetSectionBase().X + Comp->ComponentSizeQuads, Comp->GetSectionBase().Y + Comp->ComponentSizeQuads, (uint16*)HeightData.GetData(), 0, (uint16*)NormalData.GetData());
 
 		// Construct the heightmap textures
 		UObject* TextureOuter = bMoveToCurrentLevel ? Comp->GetWorld()->GetCurrentLevel()->GetOutermost() : nullptr;
@@ -126,7 +126,7 @@ void ALandscape::SplitHeightmap(ULandscapeComponent* Comp, bool bMoveToCurrentLe
 		{
 			HeightmapTexture->Source.UnlockMip(i);
 		}
-		LandscapeEdit.SetHeightData(Comp->GetSectionBase().X, Comp->GetSectionBase().Y, Comp->GetSectionBase().X + Comp->ComponentSizeQuads, Comp->GetSectionBase().Y + Comp->ComponentSizeQuads, (uint16*)HeightData.GetTypedData(), 0, false, (uint16*)NormalData.GetTypedData());
+		LandscapeEdit.SetHeightData(Comp->GetSectionBase().X, Comp->GetSectionBase().Y, Comp->GetSectionBase().X + Comp->ComponentSizeQuads, Comp->GetSectionBase().Y + Comp->ComponentSizeQuads, (uint16*)HeightData.GetData(), 0, false, (uint16*)NormalData.GetData());
 	}
 
 	// End of LandscapeEdit interface
@@ -2385,7 +2385,7 @@ void FEdModeLandscape::ImportData(const FLandscapeTargetListInfo& TargetInfo, co
 				if (Data.Num() == (1 + MaxX - MinX)*(1 + MaxY - MinY) * 2)
 				{
 					FHeightmapAccessor<false> HeightmapAccessor(LandscapeInfo);
-					HeightmapAccessor.SetData(MinX, MinY, MaxX, MaxY, (uint16*)Data.GetTypedData());
+					HeightmapAccessor.SetData(MinX, MinY, MaxX, MaxY, (uint16*)Data.GetData());
 				}
 				else
 				{
@@ -2398,7 +2398,7 @@ void FEdModeLandscape::ImportData(const FLandscapeTargetListInfo& TargetInfo, co
 				if (Data.Num() == (1 + MaxX - MinX)*(1 + MaxY - MinY))
 				{
 					FAlphamapAccessor<false, true> AlphamapAccessor(LandscapeInfo, TargetInfo.LayerInfoObj.Get());
-					AlphamapAccessor.SetData(MinX, MinY, MaxX, MaxY, Data.GetTypedData(), ELandscapeLayerPaintingRestriction::None);
+					AlphamapAccessor.SetData(MinX, MinY, MaxX, MaxY, Data.GetData(), ELandscapeLayerPaintingRestriction::None);
 				}
 				else
 				{
@@ -2608,7 +2608,7 @@ bool LandscapeEditorUtils::SetHeightmapData(ALandscapeProxy* Landscape, const TA
 	if (Data.Num() == (1 + ComponentsRect.Width())*(1 + ComponentsRect.Height()))
 	{
 		FHeightmapAccessor<false> HeightmapAccessor(Landscape->GetLandscapeInfo());
-		HeightmapAccessor.SetData(ComponentsRect.Min.X, ComponentsRect.Min.Y, ComponentsRect.Max.X, ComponentsRect.Max.Y, Data.GetTypedData());
+		HeightmapAccessor.SetData(ComponentsRect.Min.X, ComponentsRect.Min.Y, ComponentsRect.Max.X, ComponentsRect.Max.Y, Data.GetData());
 		return true;
 	}
 
@@ -2622,7 +2622,7 @@ bool LandscapeEditorUtils::SetWeightmapData(ALandscapeProxy* Landscape, ULandsca
 	if (Data.Num() == (1 + ComponentsRect.Width())*(1 + ComponentsRect.Height()))
 	{
 		FAlphamapAccessor<false, true> AlphamapAccessor(Landscape->GetLandscapeInfo(), LayerObject);
-		AlphamapAccessor.SetData(ComponentsRect.Min.X, ComponentsRect.Min.Y, ComponentsRect.Max.X, ComponentsRect.Max.Y, Data.GetTypedData(), ELandscapeLayerPaintingRestriction::None);
+		AlphamapAccessor.SetData(ComponentsRect.Min.X, ComponentsRect.Min.Y, ComponentsRect.Max.X, ComponentsRect.Max.Y, Data.GetData(), ELandscapeLayerPaintingRestriction::None);
 		return true;
 	}
 
diff --git a/Engine/Source/Editor/LandscapeEditor/Private/LandscapeEdModeRampTool.cpp b/Engine/Source/Editor/LandscapeEditor/Private/LandscapeEdModeRampTool.cpp
--- a/Engine/Source/Editor/LandscapeEditor/Private/LandscapeEdModeRampTool.cpp
+++ b/Engine/Source/Editor/LandscapeEditor/Private/LandscapeEdModeRampTool.cpp
@@ -218,7 +218,7 @@ public:
 				int32 ValidMinY = MinY;
 				int32 ValidMaxX = MaxX;
 				int32 ValidMaxY = MaxY;
-				LandscapeEdit.GetHeightData(ValidMinX, ValidMinY, ValidMaxX, ValidMaxY, Data.GetTypedData(), 0);
+				LandscapeEdit.GetHeightData(ValidMinX, ValidMinY, ValidMaxX, ValidMaxY, Data.GetData(), 0);
 
 				if (ValidMaxX - ValidMinX != 1 && ValidMaxY - ValidMinY != 1)
 				{
@@ -472,7 +472,7 @@ public:
 			int32 ValidMinY = MinY;
 			int32 ValidMaxX = MaxX;
 			int32 ValidMaxY = MaxY;
-			LandscapeEdit.GetHeightData(ValidMinX, ValidMinY, ValidMaxX, ValidMaxY, Data.GetTypedData(), 0);
+			LandscapeEdit.GetHeightData(ValidMinX, ValidMinY, ValidMaxX, ValidMaxY, Data.GetData(), 0);
 
 			if (ValidMinX > ValidMaxX || ValidMinY > ValidMaxY)
 			{
@@ -502,7 +502,7 @@ public:
 			Rasterizer.DrawTriangle(FVector2D(1, Heights[0]), FVector2D(0, Heights[0]), FVector2D(1, Heights[1]), InnerVerts[0][1], OuterVerts[0][1], InnerVerts[1][1], false);
 			Rasterizer.DrawTriangle(FVector2D(0, Heights[0]), FVector2D(1, Heights[1]), FVector2D(0, Heights[1]), OuterVerts[0][1], InnerVerts[1][1], OuterVerts[1][1], false);
 
-			LandscapeEdit.SetHeightData(MinX, MinY, MaxX, MaxY, Data.GetTypedData(), 0, true);
+			LandscapeEdit.SetHeightData(MinX, MinY, MaxX, MaxY, Data.GetData(), 0, true);
 			LandscapeEdit.Flush();
 
 			TSet<ULandscapeComponent*> Components;
diff --git a/Engine/Source/Editor/LandscapeEditor/Private/LandscapeEdModeTools.h b/Engine/Source/Editor/LandscapeEditor/Private/LandscapeEdModeTools.h
--- a/Engine/Source/Editor/LandscapeEditor/Private/LandscapeEdModeTools.h
+++ b/Engine/Source/Editor/LandscapeEditor/Private/LandscapeEdModeTools.h
@@ -486,7 +486,7 @@ struct TLandscapeEditCache
 		}
 
 		// Update real data
-		DataAccess.SetData(X1, Y1, X2, Y2, Data.GetTypedData(), PaintingRestriction);
+		DataAccess.SetData(X1, Y1, X2, Y2, Data.GetData(), PaintingRestriction);
 	}
 
 	// Get the original data before we made any changes with the SetCachedData interface.
@@ -793,7 +793,7 @@ struct FXYOffsetmapAccessor
 
 				// Update landscape.
 				LandscapeEdit->SetXYOffsetData(X1, Y1, X2, Y2, Data, 0); // XY Offset always need to be update before the height update
-				LandscapeEdit->SetHeightData(X1, Y1, X2, Y2, NewHeights.GetTypedData(), 0, true);
+				LandscapeEdit->SetHeightData(X1, Y1, X2, Y2, NewHeights.GetData(), 0, true);
 
 				// Snap foliage for each component.
 				for (int32 Index = 0; Index < CollisionComponents.Num(); ++Index)
@@ -807,7 +807,7 @@ struct FXYOffsetmapAccessor
 			{
 				// No foliage, just update landscape.
 				LandscapeEdit->SetXYOffsetData(X1, Y1, X2, Y2, Data, 0); // XY Offset always need to be update before the height update
-				LandscapeEdit->SetHeightData(X1, Y1, X2, Y2, NewHeights.GetTypedData(), 0, true);
+				LandscapeEdit->SetHeightData(X1, Y1, X2, Y2, NewHeights.GetData(), 0, true);
 			}
 		}
 		else
@@ -1084,7 +1084,7 @@ struct FLandscapeFullWeightCache : public TLandscapeEditCache < FFullWeightmapAc
 		}
 
 		// Update real data
-		DataAccess.SetData(X1, Y1, X2, Y2, Data.GetTypedData(), PaintingRestriction);
+		DataAccess.SetData(X1, Y1, X2, Y2, Data.GetData(), PaintingRestriction);
 	}
 
 	void AddDirtyLayer(ULandscapeLayerInfoObject* LayerInfo)
diff --git a/Engine/Source/Editor/LandscapeEditor/Private/LandscapeEditorDetailCustomization_CopyPaste.cpp b/Engine/Source/Editor/LandscapeEditor/Private/LandscapeEditorDetailCustomization_CopyPaste.cpp
--- a/Engine/Source/Editor/LandscapeEditor/Private/LandscapeEditorDetailCustomization_CopyPaste.cpp
+++ b/Engine/Source/Editor/LandscapeEditor/Private/LandscapeEditorDetailCustomization_CopyPaste.cpp
@@ -353,7 +353,7 @@ FReply FLandscapeEditorDetailCustomization_CopyPaste::OnGizmoImportButtonClicked
 				}
 			}
 
-			Gizmo->Import(LandscapeEdMode->UISettings->GizmoImportSize.X, LandscapeEdMode->UISettings->GizmoImportSize.Y, (uint16*)Data.GetTypedData(), LayerInfos, LayerDataPtrs.Num() ? LayerDataPtrs.GetTypedData() : NULL);
+			Gizmo->Import(LandscapeEdMode->UISettings->GizmoImportSize.X, LandscapeEdMode->UISettings->GizmoImportSize.Y, (uint16*)Data.GetData(), LayerInfos, LayerDataPtrs.Num() ? LayerDataPtrs.GetData() : NULL);
 		}
 	}
 
diff --git a/Engine/Source/Editor/LandscapeEditor/Private/LandscapeEditorObject.cpp b/Engine/Source/Editor/LandscapeEditor/Private/LandscapeEditorObject.cpp
--- a/Engine/Source/Editor/LandscapeEditor/Private/LandscapeEditorObject.cpp
+++ b/Engine/Source/Editor/LandscapeEditor/Private/LandscapeEditorObject.cpp
@@ -473,16 +473,16 @@ bool ULandscapeEditorObject::SetAlphaTexture(UTexture2D* InTexture, EColorChanne
 		switch(AlphaTextureChannel)
 		{
 		case 1:
-			SrcPtr = &((FColor*)NewTextureData.GetTypedData())->G;
+			SrcPtr = &((FColor*)NewTextureData.GetData())->G;
 			break;
 		case 2:
-			SrcPtr = &((FColor*)NewTextureData.GetTypedData())->B;
+			SrcPtr = &((FColor*)NewTextureData.GetData())->B;
 			break;
 		case 3:
-			SrcPtr = &((FColor*)NewTextureData.GetTypedData())->A;
+			SrcPtr = &((FColor*)NewTextureData.GetData())->A;
 			break;
 		default:
-			SrcPtr = &((FColor*)NewTextureData.GetTypedData())->R;
+			SrcPtr = &((FColor*)NewTextureData.GetData())->R;
 			break;
 		}
 
diff --git a/Engine/Source/Editor/MeshPaint/Private/MeshPaintEdMode.cpp b/Engine/Source/Editor/MeshPaint/Private/MeshPaintEdMode.cpp
--- a/Engine/Source/Editor/MeshPaint/Private/MeshPaintEdMode.cpp
+++ b/Engine/Source/Editor/MeshPaint/Private/MeshPaintEdMode.cpp
@@ -2452,7 +2452,7 @@ void FEdModeMeshPaint::CommitAllPaintedTextures()
 					// Store source art
 					FColor* Colors = (FColor*)TextureData->PaintingTexture2D->Source.LockMip(0);					
 					check(TextureData->PaintingTexture2D->Source.CalcMipSize(0)==TexturePixels.Num()*sizeof(FColor));
-					FMemory::Memcpy(Colors, TexturePixels.GetTypedData(), TexturePixels.Num() * sizeof(FColor));
+					FMemory::Memcpy(Colors, TexturePixels.GetData(), TexturePixels.Num() * sizeof(FColor));
 					TextureData->PaintingTexture2D->Source.UnlockMip(0);
 
 					// If render target gamma used was 1.0 then disable SRGB for the static texture
@@ -2666,7 +2666,7 @@ void FEdModeMeshPaint::FinishPaintingTexture( )
 				// Store source art
 				FColor* Colors = (FColor*)TextureData->PaintingTexture2D->Source.LockMip(0);
 				check(TextureData->PaintingTexture2D->Source.CalcMipSize(0)==TexturePixels.Num()*sizeof(FColor));
-				FMemory::Memcpy(Colors, TexturePixels.GetTypedData(), TexturePixels.Num() * sizeof(FColor));
+				FMemory::Memcpy(Colors, TexturePixels.GetData(), TexturePixels.Num() * sizeof(FColor));
 				TextureData->PaintingTexture2D->Source.UnlockMip(0);
 
 				// If render target gamma used was 1.0 then disable SRGB for the static texture
@@ -3197,7 +3197,7 @@ UTexture2D* FEdModeMeshPaint::CreateTempUncompressedTexture( UTexture2D* SourceT
 	for( int32 y=0; y<Height; y++ )
 	{
 		uint8* DestPtr = &MipData[(Height - 1 - y) * Width * sizeof(FColor)];
-		const FColor* SrcPtr = &( (FColor*)( RawData.GetTypedData() ) )[ ( Height - 1 - y ) * Width ];
+		const FColor* SrcPtr = &( (FColor*)( RawData.GetData() ) )[ ( Height - 1 - y ) * Width ];
 		for( int32 x=0; x<Width; x++ )
 		{
 			*DestPtr++ = SrcPtr->B;
@@ -4319,7 +4319,7 @@ void FImportVertexTextureHelper::ImportVertexColors(const FString & Filename, in
 
 	TArray<uint8> SrcMipData;
 	Tex->Source.GetMipData(SrcMipData, 0);
-	uint8* MipData = SrcMipData.GetTypedData();
+	uint8* MipData = SrcMipData.GetData();
 	TArray <UStaticMesh*> ModifiedStaticMeshes;
 
 	for(const auto& StaticMeshComponent : Components)
@@ -5000,7 +5000,7 @@ void FEdModeMeshPaint::DuplicateTextureMaterialCombo()
 						SelectedTexture->Source.GetMipData(TexturePixels, 0);					
 						uint8* DestData = NewTexture->Source.LockMip(0);					
 						check(NewTexture->Source.CalcMipSize(0)==TexturePixels.Num()*sizeof(uint8));
-						FMemory::Memcpy(DestData, TexturePixels.GetTypedData(), TexturePixels.Num() * sizeof( uint8 ) );
+						FMemory::Memcpy(DestData, TexturePixels.GetData(), TexturePixels.Num() * sizeof( uint8 ) );
 						NewTexture->Source.UnlockMip(0);
 						NewTexture->SRGB = SelectedTexture->SRGB;
 						NewTexture->PostEditChange();
diff --git a/Engine/Source/Editor/PropertyEditor/Private/PropertyChangeListener.cpp b/Engine/Source/Editor/PropertyEditor/Private/PropertyChangeListener.cpp
--- a/Engine/Source/Editor/PropertyEditor/Private/PropertyChangeListener.cpp
+++ b/Engine/Source/Editor/PropertyEditor/Private/PropertyChangeListener.cpp
@@ -106,13 +106,13 @@ public:
 			{
 				// Check static arrays
 				Data.AddZeroed( Property->ArrayDim * Property->ElementSize );
-				Property->CopyCompleteValue( Data.GetTypedData(), PropertyValueAddress);
+				Property->CopyCompleteValue( Data.GetData(), PropertyValueAddress);
 			}
 			else
 			{
 				// Regular properties
 				Data.AddZeroed( Property->ElementSize );
-				Property->CopySingleValue( Data.GetTypedData(), PropertyValueAddress );
+				Property->CopySingleValue( Data.GetData(), PropertyValueAddress );
 			}
 		}
 	}
@@ -143,7 +143,7 @@ public:
 
 		if( bPropertyValid )
 		{
-			bChanged = !Property->Identical( PropertyValueAddress, Data.GetTypedData() );
+			bChanged = !Property->Identical( PropertyValueAddress, Data.GetData() );
 			if( bRecacheNewValues )
 			{
 				CacheValue();
diff --git a/Engine/Source/Editor/SceneOutliner/Private/SOutlinerTreeView.cpp b/Engine/Source/Editor/SceneOutliner/Private/SOutlinerTreeView.cpp
--- a/Engine/Source/Editor/SceneOutliner/Private/SOutlinerTreeView.cpp
+++ b/Engine/Source/Editor/SceneOutliner/Private/SOutlinerTreeView.cpp
@@ -311,7 +311,7 @@ namespace SceneOutliner
 				{
 					if (DraggedActors->Num() == 1)
 					{
-						const AActor* Actor = DraggedActors->GetTypedData()[0].Get();
+						const AActor* Actor = DraggedActors->GetData()[0].Get();
 						if (Actor)
 						{
 							Args.Add(TEXT("SourceName"), FText::FromString(Actor->GetActorLabel()));
@@ -319,7 +319,7 @@ namespace SceneOutliner
 					}
 					else
 					{
-						const auto DraggedFolder = DraggedFolders->GetTypedData()[0].Pin();
+						const auto DraggedFolder = DraggedFolders->GetData()[0].Pin();
 						if (DraggedFolder.IsValid())
 						{
 							Args.Add(TEXT("SourceName"), FText::FromString(DraggedFolder->LeafName.ToString()));
diff --git a/Engine/Source/Editor/UnrealEd/Private/ApexClothingUtils.cpp b/Engine/Source/Editor/UnrealEd/Private/ApexClothingUtils.cpp
--- a/Engine/Source/Editor/UnrealEd/Private/ApexClothingUtils.cpp
+++ b/Engine/Source/Editor/UnrealEd/Private/ApexClothingUtils.cpp
@@ -263,7 +263,7 @@ bool LoadPhysicalMeshFromClothingAsset(NxClothingAsset& ApexClothingAsset,
 	uint32 MapSize = ApexClothingAsset.getMeshSkinningMapSize(LODIndex);
 
 	SkinningMap.AddUninitialized(MapSize);
-	ApexClothingAsset.getMeshSkinningMap(LODIndex,SkinningMap.GetTypedData());
+	ApexClothingAsset.getMeshSkinningMap(LODIndex,SkinningMap.GetData());
 
 	// Mapping
 	uint32 EndTotalIndex = StartSimulIndex + NumTotalVertices;
@@ -622,7 +622,7 @@ bool LoadGraphicalMeshFromClothingAsset(NxClothingAsset& ApexClothingAsset,
 		TArray<FApexBoneIndices> BoneIndices;
 		BoneIndices.AddUninitialized(SubmeshVertexCount);
 
-		bHaveBoneIndices = VB.getBufferData(BoneIndices.GetTypedData(), physx::NxRenderDataFormat::USHORT4, sizeof(FApexBoneIndices), BoneIndexBufferIndex, 0, SubmeshVertexCount);
+		bHaveBoneIndices = VB.getBufferData(BoneIndices.GetData(), physx::NxRenderDataFormat::USHORT4, sizeof(FApexBoneIndices), BoneIndexBufferIndex, 0, SubmeshVertexCount);
 		// Converting uint16 to uint8 because APEX Clothing is using uint16 for Bone Index 
 		if (bHaveBoneIndices)
 		{
@@ -641,7 +641,7 @@ bool LoadGraphicalMeshFromClothingAsset(NxClothingAsset& ApexClothingAsset,
 		TArray<uint16> BoneIndices;
 		BoneIndices.AddUninitialized(SubmeshVertexCount);
 
-		bHaveBoneIndices = VB.getBufferData(BoneIndices.GetTypedData(), physx::NxRenderDataFormat::USHORT1, sizeof(uint16), BoneIndexBufferIndex, 0, SubmeshVertexCount);
+		bHaveBoneIndices = VB.getBufferData(BoneIndices.GetData(), physx::NxRenderDataFormat::USHORT1, sizeof(uint16), BoneIndexBufferIndex, 0, SubmeshVertexCount);
 
 		if(bHaveBoneIndices)
 		{
@@ -672,7 +672,7 @@ bool LoadGraphicalMeshFromClothingAsset(NxClothingAsset& ApexClothingAsset,
 		TArray<float> BoneWeights;
 		BoneWeights.AddUninitialized(SubmeshVertexCount);
 
-		bHaveBoneWeights = VB.getBufferData(BoneWeights.GetTypedData(), physx::NxRenderDataFormat::FLOAT1, sizeof(float), BoneWeightBufferIndex, 0, SubmeshVertexCount);
+		bHaveBoneWeights = VB.getBufferData(BoneWeights.GetData(), physx::NxRenderDataFormat::FLOAT1, sizeof(float), BoneWeightBufferIndex, 0, SubmeshVertexCount);
 
 		if(bHaveBoneWeights)
 		{
@@ -1096,7 +1096,7 @@ NxClothingAsset* ApplyTransform(NxClothingAsset* ApexClothingAsset)
 			NewBindPoses.Add(NewBindPose);
 		}
 
-		ApexClothingAssetAuthoring->updateBindPoses(NewBindPoses.GetTypedData(), NewBindPoses.Num(), true, true);
+		ApexClothingAssetAuthoring->updateBindPoses(NewBindPoses.GetData(), NewBindPoses.Num(), true, true);
 
 		// destroy and create a new asset based off the authoring version
 	
diff --git a/Engine/Source/Editor/UnrealEd/Private/Commandlets/DiffPackagesCommandlet.cpp b/Engine/Source/Editor/UnrealEd/Private/Commandlets/DiffPackagesCommandlet.cpp
--- a/Engine/Source/Editor/UnrealEd/Private/Commandlets/DiffPackagesCommandlet.cpp
+++ b/Engine/Source/Editor/UnrealEd/Private/Commandlets/DiffPackagesCommandlet.cpp
@@ -740,7 +740,7 @@ int32 UDiffPackagesCommandlet::Main(const FString& Params)
 			{
 				if ( Comparison.ObjectSets[PackageIndex] != NULL )
 				{
-					AllObjectSets[PackageIndex].AddRawItem(Comparison.ObjectSets[PackageIndex]);
+					AllObjectSets[PackageIndex].Add(Comparison.ObjectSets[PackageIndex]);
 				}
 			}
 
@@ -763,7 +763,7 @@ int32 UDiffPackagesCommandlet::Main(const FString& Params)
 				{
 					if ( Comparison.ObjectSets[PackageIndex] != NULL )
 					{
-						AllObjectSets[PackageIndex].AddRawItem(Comparison.ObjectSets[PackageIndex]);
+						AllObjectSets[PackageIndex].Add(Comparison.ObjectSets[PackageIndex]);
 					}
 				}
 
diff --git a/Engine/Source/Editor/UnrealEd/Private/Commandlets/PackageUtilities.cpp b/Engine/Source/Editor/UnrealEd/Private/Commandlets/PackageUtilities.cpp
--- a/Engine/Source/Editor/UnrealEd/Private/Commandlets/PackageUtilities.cpp
+++ b/Engine/Source/Editor/UnrealEd/Private/Commandlets/PackageUtilities.cpp
@@ -1597,12 +1597,12 @@ struct CompressAnimationsFunctor
  						// Translation
 						{
 							// Use the CompressedTrackOffsets stream to find the data addresses
- 							const int32* RESTRICT TrackDataForTransKey = AnimSeq->CompressedTrackOffsets.GetTypedData() + (TrackIndex * 2);
+							const int32* RESTRICT TrackDataForTransKey = AnimSeq->CompressedTrackOffsets.GetData() + (TrackIndex * 2);
 							const int32 TransKeysOffset = TrackDataForTransKey[0];
  							if( TransKeysOffset != INDEX_NONE )
  							{
- 								const uint8* RESTRICT TrackData = AnimSeq->CompressedByteStream.GetTypedData() + TransKeysOffset + 4;
- 								const int32 Header = *((int32*)(AnimSeq->CompressedByteStream.GetTypedData() + TransKeysOffset));
+								const uint8* RESTRICT TrackData = AnimSeq->CompressedByteStream.GetData() + TransKeysOffset + 4;
+								const int32 Header = *((int32*)(AnimSeq->CompressedByteStream.GetData() + TransKeysOffset));
 	 
  								int32 KeyFormat;
  								int32 NumKeys;
@@ -1705,12 +1705,12 @@ struct CompressAnimationsFunctor
 						// Rotation
 						{
 							// Use the CompressedTrackOffsets stream to find the data addresses
-							const int32* RESTRICT TrackDataForRotKey = AnimSeq->CompressedTrackOffsets.GetTypedData() + (TrackIndex * 2);
+							const int32* RESTRICT TrackDataForRotKey = AnimSeq->CompressedTrackOffsets.GetData() + (TrackIndex * 2);
 							const int32 RotKeysOffset = TrackDataForRotKey[1];
 							if( RotKeysOffset != INDEX_NONE )
 							{
-								const uint8* RESTRICT TrackData = AnimSeq->CompressedByteStream.GetTypedData() + RotKeysOffset + 4;
-								const int32 Header = *((int32*)(AnimSeq->CompressedByteStream.GetTypedData() + RotKeysOffset));
+								const uint8* RESTRICT TrackData = AnimSeq->CompressedByteStream.GetData() + RotKeysOffset + 4;
+								const int32 Header = *((int32*)(AnimSeq->CompressedByteStream.GetData() + RotKeysOffset));
 
 								int32 KeyFormat;
 								int32 NumKeys;
@@ -1769,8 +1769,8 @@ struct CompressAnimationsFunctor
 							const int32 ScaleKeysOffset = AnimSeq->CompressedScaleOffsets.GetOffsetData(TrackIndex, 0);
 							if( ScaleKeysOffset != INDEX_NONE )
 							{
-								const uint8* RESTRICT TrackData = AnimSeq->CompressedByteStream.GetTypedData() + ScaleKeysOffset + 4;
-								const int32 Header = *((int32*)(AnimSeq->CompressedByteStream.GetTypedData() + ScaleKeysOffset));
+								const uint8* RESTRICT TrackData = AnimSeq->CompressedByteStream.GetData() + ScaleKeysOffset + 4;
+								const int32 Header = *((int32*)(AnimSeq->CompressedByteStream.GetData() + ScaleKeysOffset));
 
 								int32 KeyFormat;
 								int32 NumKeys;
diff --git a/Engine/Source/Editor/UnrealEd/Private/CubemapUnwrapUtils.cpp b/Engine/Source/Editor/UnrealEd/Private/CubemapUnwrapUtils.cpp
--- a/Engine/Source/Editor/UnrealEd/Private/CubemapUnwrapUtils.cpp
+++ b/Engine/Source/Editor/UnrealEd/Private/CubemapUnwrapUtils.cpp
@@ -69,13 +69,13 @@ namespace CubemapHelpers
 		switch (TargetPixelFormat)
 		{
 			case PF_B8G8R8A8:
-				bReadSuccess = RenderTarget->ReadPixelsPtr((FColor*)BitsOUT.GetTypedData());
+				bReadSuccess = RenderTarget->ReadPixelsPtr((FColor*)BitsOUT.GetData());
 			break;
 			case PF_FloatRGBA:
 				{
 					TArray<FFloat16Color> FloatColors;
 					bReadSuccess = RenderTarget->ReadFloat16Pixels(FloatColors);
-					FMemory::Memcpy(BitsOUT.GetTypedData(), FloatColors.GetTypedData(), ImageBytes);
+					FMemory::Memcpy(BitsOUT.GetData(), FloatColors.GetData(), ImageBytes);
 				}
 			break;
 		}
diff --git a/Engine/Source/Editor/UnrealEd/Private/Editor.cpp b/Engine/Source/Editor/UnrealEd/Private/Editor.cpp
--- a/Engine/Source/Editor/UnrealEd/Private/Editor.cpp
+++ b/Engine/Source/Editor/UnrealEd/Private/Editor.cpp
@@ -2936,14 +2936,14 @@ void SoundWaveQualityPreview( USoundWave* SoundWave, FPreviewInfo* PreviewInfo )
 	{
 		TArray<uint8> Input;
 		Input.AddUninitialized(WaveInfo.SampleDataSize);
-		FMemory::Memcpy(Input.GetTypedData(), WaveInfo.SampleDataStart, Input.Num());
+		FMemory::Memcpy(Input.GetData(), WaveInfo.SampleDataStart, Input.Num());
 		TArray<uint8> Output;
 		Compressor->Recompress(NAME_OGG, Input, QualityInfo, Output );
 		if (Output.Num())
 		{
 			PreviewInfo->OggVorbisSize = Output.Num();
 			PreviewInfo->DecompressedOggVorbis = ( uint8* )FMemory::Malloc(Output.Num());
-			FMemory::Memcpy(PreviewInfo->DecompressedOggVorbis, Output.GetTypedData(), Output.Num());
+			FMemory::Memcpy(PreviewInfo->DecompressedOggVorbis, Output.GetData(), Output.Num());
 		}
 	}
 
diff --git a/Engine/Source/Editor/UnrealEd/Private/EditorExporters.cpp b/Engine/Source/Editor/UnrealEd/Private/EditorExporters.cpp
--- a/Engine/Source/Editor/UnrealEd/Private/EditorExporters.cpp
+++ b/Engine/Source/Editor/UnrealEd/Private/EditorExporters.cpp
@@ -1463,12 +1463,12 @@ static void AddActorToOBJs(AActor* Actor, TArray<FOBJGeom*>& Objects, TSet<UMate
 				{
 					TexIndex = AllocInfo.WeightmapTextureIndex;
 					Component->WeightmapTextures[TexIndex]->Source.GetMipData(RawVisData, 0);
-					VisDataMap = RawVisData.GetTypedData() + ChannelOffsets[AllocInfo.WeightmapTextureChannel];
+					VisDataMap = RawVisData.GetData() + ChannelOffsets[AllocInfo.WeightmapTextureChannel];
 				}
 			}
 
 			// Export verts
-			FOBJVertex* Vert = OBJGeom->VertexData.GetTypedData();
+			FOBJVertex* Vert = OBJGeom->VertexData.GetData();
 			for( int32 y=0;y<ComponentSizeQuads+1;y++ )
 			{
 				for( int32 x=0;x<ComponentSizeQuads+1;x++ )
@@ -1487,7 +1487,7 @@ static void AddActorToOBJs(AActor* Actor, TArray<FOBJGeom*>& Objects, TSet<UMate
 			int32 VisThreshold = 170;
 			int32 SubNumX, SubNumY, SubX, SubY;
 
-			FOBJFace* Face = OBJGeom->Faces.GetTypedData();
+			FOBJFace* Face = OBJGeom->Faces.GetData();
 			for( int32 y=0;y<ComponentSizeQuads;y++ )
 			{
 				for( int32 x=0;x<ComponentSizeQuads;x++ )
@@ -1552,7 +1552,7 @@ static void AddActorToOBJs(AActor* Actor, TArray<FOBJGeom*>& Objects, TSet<UMate
 		
 		uint32 VertexCount = RenderData->PositionVertexBuffer.GetNumVertices();
 		OBJGeom->VertexData.AddUninitialized(VertexCount);
-		FOBJVertex* VerticesOut = OBJGeom->VertexData.GetTypedData();
+		FOBJVertex* VerticesOut = OBJGeom->VertexData.GetData();
 
 		check(VertexCount == RenderData->VertexBuffer.GetNumVertices());
 
@@ -1655,7 +1655,7 @@ static void ExportMaterialPropertyTexture(
 	}
 
 	// export the diffuse channel bmp
-	FFileHelper::CreateBitmap(*BMPFilename, SizeX, SizeY, OutputBMP.GetTypedData());
+	FFileHelper::CreateBitmap(*BMPFilename, SizeX, SizeY, OutputBMP.GetData());
 }
 
 /**
@@ -2803,7 +2803,7 @@ namespace MaterialExportUtils
 				// Copy the contents of the remote texture to system memory
 				// NOTE: OutRawImageData must be a preallocated buffer!
 				RenderTargetResource->ReadPixelsPtr(
-					(FColor*)OutFlattenMaterial.DiffuseSamples.GetTypedData(), 
+					(FColor*)OutFlattenMaterial.DiffuseSamples.GetData(), 
 					FReadSurfaceDataFlags(), 
 					FIntRect(0, 0, OutFlattenMaterial.DiffuseSize.X, OutFlattenMaterial.DiffuseSize.Y)
 					);
@@ -2887,7 +2887,7 @@ namespace MaterialExportUtils
 				// Copy the contents of the remote texture to system memory
 				// NOTE: OutRawImageData must be a preallocated buffer!
 				RenderTargetResource->ReadPixelsPtr(
-					(FColor*)OutFlattenMaterial.NormalSamples.GetTypedData(), 
+					(FColor*)OutFlattenMaterial.NormalSamples.GetData(), 
 					FReadSurfaceDataFlags(), 
 					FIntRect(0, 0, OutFlattenMaterial.NormalSize.X, OutFlattenMaterial.NormalSize.Y)
 					);
diff --git a/Engine/Source/Editor/UnrealEd/Private/EditorViewportClient.cpp b/Engine/Source/Editor/UnrealEd/Private/EditorViewportClient.cpp
--- a/Engine/Source/Editor/UnrealEd/Private/EditorViewportClient.cpp
+++ b/Engine/Source/Editor/UnrealEd/Private/EditorViewportClient.cpp
@@ -3778,7 +3778,7 @@ void FEditorViewportClient::TakeScreenshot(FViewport* InViewport, bool bInValida
 		{
 			// Save the contents of the array to a bitmap file.
 			FString ScreenshotSaveName = "";
-			if ( FFileHelper::CreateBitmap(*(FPaths::ScreenShotDir() / TEXT("ScreenShot")),InViewport->GetSizeXY().X,InViewport->GetSizeXY().Y,Bitmap.GetTypedData(),NULL,&IFileManager::Get(),&ScreenshotSaveName) )
+			if ( FFileHelper::CreateBitmap(*(FPaths::ScreenShotDir() / TEXT("ScreenShot")),InViewport->GetSizeXY().X,InViewport->GetSizeXY().Y,Bitmap.GetData(),NULL,&IFileManager::Get(),&ScreenshotSaveName) )
 			{
 				// Setup the string with the path and name of the file
 				ScreenshotSaveResultText = NSLOCTEXT( "UnrealEd", "ScreenshotSavedAs", "Screenshot capture saved as" );					
@@ -3911,7 +3911,7 @@ void FEditorViewportClient::ProcessScreenShots(FViewport* InViewport)
 			}
 
 			// Save the bitmap to disc
-			FFileHelper::CreateBitmap(*ScreenShotName, BitmapSize.X, BitmapSize.Y, Bitmap.GetTypedData(), &SourceRect, &IFileManager::Get(), NULL, bWriteAlpha);
+			FFileHelper::CreateBitmap(*ScreenShotName, BitmapSize.X, BitmapSize.Y, Bitmap.GetData(), &SourceRect, &IFileManager::Get(), NULL, bWriteAlpha);
 		}
 		
 		// Done with the request
diff --git a/Engine/Source/Editor/UnrealEd/Private/Factories/ApexDestructibleAssetImport.cpp b/Engine/Source/Editor/UnrealEd/Private/Factories/ApexDestructibleAssetImport.cpp
--- a/Engine/Source/Editor/UnrealEd/Private/Factories/ApexDestructibleAssetImport.cpp
+++ b/Engine/Source/Editor/UnrealEd/Private/Factories/ApexDestructibleAssetImport.cpp
@@ -507,20 +507,20 @@ static bool FillSkelMeshImporterFromApexDestructibleAsset(FSkeletalMeshImportDat
 		const PxI32 NormalBufferIndex = VBFormat.getBufferIndexFromID(VBFormat.getSemanticID(NxRenderVertexSemantic::NORMAL));
 		TArray<FVector> Normals;
 		Normals.AddUninitialized(SubmeshVertexCount);
-		const bool bHaveNormals = VB.getBufferData(Normals.GetTypedData(), physx::NxRenderDataFormat::FLOAT3, sizeof(FVector), NormalBufferIndex, 0, SubmeshVertexCount);
+		const bool bHaveNormals = VB.getBufferData(Normals.GetData(), physx::NxRenderDataFormat::FLOAT3, sizeof(FVector), NormalBufferIndex, 0, SubmeshVertexCount);
 		if (!bHaveNormals)
 		{
-			FMemory::Memset(Normals.GetTypedData(), 0, SubmeshVertexCount*sizeof(FVector));	// Fill with zeros
+			FMemory::Memset(Normals.GetData(), 0, SubmeshVertexCount*sizeof(FVector));	// Fill with zeros
 		}
 
 		// Tangents
 		const PxI32 TangentBufferIndex = VBFormat.getBufferIndexFromID(VBFormat.getSemanticID(NxRenderVertexSemantic::TANGENT));
 		TArray<FVector> Tangents;
 		Tangents.AddUninitialized(SubmeshVertexCount);
-		const bool bHaveTangents = VB.getBufferData(Tangents.GetTypedData(), physx::NxRenderDataFormat::FLOAT3, sizeof(FVector), TangentBufferIndex, 0, SubmeshVertexCount);
+		const bool bHaveTangents = VB.getBufferData(Tangents.GetData(), physx::NxRenderDataFormat::FLOAT3, sizeof(FVector), TangentBufferIndex, 0, SubmeshVertexCount);
 		if (!bHaveTangents)
 		{
-			FMemory::Memset(Tangents.GetTypedData(), 0, SubmeshVertexCount*sizeof(FVector));	// Fill with zeros
+			FMemory::Memset(Tangents.GetData(), 0, SubmeshVertexCount*sizeof(FVector));	// Fill with zeros
 		}
 
 		// Update bHaveAllNormals and bHaveAllTangents
@@ -531,7 +531,7 @@ static bool FillSkelMeshImporterFromApexDestructibleAsset(FSkeletalMeshImportDat
 		const PxI32 BinormalBufferIndex = VBFormat.getBufferIndexFromID(VBFormat.getSemanticID(NxRenderVertexSemantic::BINORMAL));
 		TArray<FVector> Binormals;
 		Binormals.AddUninitialized(SubmeshVertexCount);
-		bool bHaveBinormals = VB.getBufferData(Binormals.GetTypedData(), physx::NxRenderDataFormat::FLOAT3, sizeof(FVector), BinormalBufferIndex, 0, SubmeshVertexCount);
+		bool bHaveBinormals = VB.getBufferData(Binormals.GetData(), physx::NxRenderDataFormat::FLOAT3, sizeof(FVector), BinormalBufferIndex, 0, SubmeshVertexCount);
 		if (!bHaveBinormals)
 		{
 			bHaveBinormals = bHaveNormals && bHaveTangents;
@@ -545,10 +545,10 @@ static bool FillSkelMeshImporterFromApexDestructibleAsset(FSkeletalMeshImportDat
 		const PxI32 ColorBufferIndex = VBFormat.getBufferIndexFromID(VBFormat.getSemanticID(NxRenderVertexSemantic::COLOR));
 		TArray<FColor> Colors;
 		Colors.AddUninitialized(SubmeshVertexCount);
-		const bool bHaveColors = VB.getBufferData(Colors.GetTypedData(), physx::NxRenderDataFormat::B8G8R8A8, sizeof(FColor), ColorBufferIndex, 0, SubmeshVertexCount);
+		const bool bHaveColors = VB.getBufferData(Colors.GetData(), physx::NxRenderDataFormat::B8G8R8A8, sizeof(FColor), ColorBufferIndex, 0, SubmeshVertexCount);
 		if (!bHaveColors)
 		{
-			FMemory::Memset(Colors.GetTypedData(), 0xFF, SubmeshVertexCount*sizeof(FColor));	// Fill with 0xFF
+			FMemory::Memset(Colors.GetData(), 0xFF, SubmeshVertexCount*sizeof(FColor));	// Fill with 0xFF
 		}
 
 		// UVs
@@ -829,7 +829,7 @@ bool SetApexDestructibleAsset(UDestructibleMesh& DestructibleMesh, NxDestructibl
 	DestructibleMesh.LODInfo[0].LODHysteresis = 0.02f;
 
 	// Create initial bounding box based on expanded version of reference pose for meshes without physics assets. Can be overridden by artist.
-	FBox BoundingBox( SkelMeshImportDataPtr->Points.GetTypedData(), SkelMeshImportDataPtr->Points.Num() );
+	FBox BoundingBox(SkelMeshImportDataPtr->Points.GetData(), SkelMeshImportDataPtr->Points.Num());
 	FBox Temp = BoundingBox;
 	FVector MidMesh		= 0.5f*(Temp.Min + Temp.Max);
 	BoundingBox.Min		= Temp.Min + 1.0f*(Temp.Min - MidMesh);
