Replace UE4 references with UE in audio code

diff --git a/Engine/Source/Developer/AudioFormatOpus/Private/AudioFormatOpus.cpp b/Engine/Source/Developer/AudioFormatOpus/Private/AudioFormatOpus.cpp
--- a/Engine/Source/Developer/AudioFormatOpus/Private/AudioFormatOpus.cpp
+++ b/Engine/Source/Developer/AudioFormatOpus/Private/AudioFormatOpus.cpp
@@ -16,8 +16,8 @@ THIRD_PARTY_INCLUDES_START
 #include "speex_resampler.h"
 THIRD_PARTY_INCLUDES_END
 
-/** Use UE4 memory allocation or Opus */
-#define USE_UE4_MEM_ALLOC 1
+/** Use UE memory allocation or Opus */
+#define USE_UE_MEM_ALLOC 1
 #define SAMPLE_SIZE			( ( uint32 )sizeof( short ) )
 
 static FName NAME_OPUS(TEXT("OPUS"));
@@ -82,7 +82,7 @@ public:
 		// Initialise the Opus encoder
 		OpusEncoder* Encoder = NULL;
 		int32 EncError = 0;
-#if USE_UE4_MEM_ALLOC
+#if USE_UE_MEM_ALLOC
 		int32 EncSize = opus_encoder_get_size(QualityInfo.NumChannels);
 		Encoder = (OpusEncoder*)FMemory::Malloc(EncSize);
 		EncError = opus_encoder_init(Encoder, kOpusSampleRate, QualityInfo.NumChannels, OPUS_APPLICATION_AUDIO);
@@ -219,7 +219,7 @@ public:
 		int32 mapping_family = 1;
 		TArray<uint8> mapping;
 		mapping.AddUninitialized(QualityInfo.NumChannels);
-#if USE_UE4_MEM_ALLOC
+#if USE_UE_MEM_ALLOC
 		int32 EncSize = opus_multistream_surround_encoder_get_size(QualityInfo.NumChannels, mapping_family);
 		Encoder = (OpusMSEncoder*)FMemory::Malloc(EncSize);
 		EncError = opus_multistream_surround_encoder_init(Encoder, kOpusSampleRate, QualityInfo.NumChannels, mapping_family, &streams, &coupled_streams, mapping.GetData(), OPUS_APPLICATION_AUDIO);
@@ -350,7 +350,7 @@ public:
 
 	virtual int32 GetMinimumSizeForInitialChunk(FName Format, const TArray<uint8>& SrcBuffer) const override
 	{
-		// Since UE4 uses it's own version of the header, we hardcode the size of that here. See SplitDataForStreaming below to see our initial info.
+		// Since UE uses it's own version of the header, we hardcode the size of that here. See SplitDataForStreaming below to see our initial info.
 		return FCStringAnsi::Strlen(OPUS_ID_STRING) + 1 // Format identifier
 			+ sizeof(uint16) // Sample Rate
 			+ sizeof(uint32) // True Sample Count
@@ -513,7 +513,7 @@ public:
 
 	void Destroy(OpusEncoder* Encoder) const
 	{
-#if USE_UE4_MEM_ALLOC
+#if USE_UE_MEM_ALLOC
 		FMemory::Free(Encoder);
 #else
 		opus_encoder_destroy(Encoder);
@@ -522,7 +522,7 @@ public:
 
 	void Destroy(OpusMSEncoder* Encoder) const
 	{
-#if USE_UE4_MEM_ALLOC
+#if USE_UE_MEM_ALLOC
 		FMemory::Free(Encoder);
 #else
 		opus_multistream_encoder_destroy(Encoder);
diff --git a/Engine/Source/Runtime/Engine/Private/OpusAudioInfo.cpp b/Engine/Source/Runtime/Engine/Private/OpusAudioInfo.cpp
--- a/Engine/Source/Runtime/Engine/Private/OpusAudioInfo.cpp
+++ b/Engine/Source/Runtime/Engine/Private/OpusAudioInfo.cpp
@@ -15,7 +15,7 @@ THIRD_PARTY_INCLUDES_START
 THIRD_PARTY_INCLUDES_END
 #endif
 
-#define USE_UE4_MEM_ALLOC 1
+#define USE_UE_MEM_ALLOC 1
 #define OPUS_MAX_FRAME_SIZE_MS 120
 
 ///////////////////////////////////////////////////////////////////////////////////////
@@ -50,7 +50,7 @@ struct FOpusDecoderWrapper
 #if WITH_OPUS
 		check(NumChannels <= 8);
 		const UnrealChannelLayout& Layout = UnrealMappings[NumChannels-1];
-	#if USE_UE4_MEM_ALLOC
+	#if USE_UE_MEM_ALLOC
 		int32 DecSize = opus_multistream_decoder_get_size(Layout.NumStreams, Layout.NumCoupledStreams);
 		Decoder = (OpusMSDecoder*)FMemory::Malloc(DecSize);
 		DecError = opus_multistream_decoder_init(Decoder, SampleRate, NumChannels, Layout.NumStreams, Layout.NumCoupledStreams, Layout.Mapping);
@@ -63,7 +63,7 @@ struct FOpusDecoderWrapper
 	~FOpusDecoderWrapper()
 	{
 #if WITH_OPUS
-	#if USE_UE4_MEM_ALLOC
+	#if USE_UE_MEM_ALLOC
 		FMemory::Free(Decoder);
 	#else
 		opus_multistream_encoder_destroy(Decoder);
diff --git a/Engine/Source/Runtime/Engine/Public/AudioDecompress.h b/Engine/Source/Runtime/Engine/Public/AudioDecompress.h
--- a/Engine/Source/Runtime/Engine/Public/AudioDecompress.h
+++ b/Engine/Source/Runtime/Engine/Public/AudioDecompress.h
@@ -185,7 +185,7 @@ struct FDecodeResult
 /** 
  * Default implementation of a streamed compressed audio format.
  * Can be subclassed to support streaming of a specific asset format. Handles all 
- * the platform independent aspects of file format streaming for you (dealing with UE4 streamed assets)
+ * the platform independent aspects of file format streaming for you (dealing with UE streamed assets)
  */
 class IStreamedCompressedInfo : public ICompressedAudioInfo
 {
diff --git a/Engine/Source/Runtime/Online/Voice/Private/VoiceCodecOpus.cpp b/Engine/Source/Runtime/Online/Voice/Private/VoiceCodecOpus.cpp
--- a/Engine/Source/Runtime/Online/Voice/Private/VoiceCodecOpus.cpp
+++ b/Engine/Source/Runtime/Online/Voice/Private/VoiceCodecOpus.cpp
@@ -12,8 +12,8 @@ THIRD_PARTY_INCLUDES_END
 #define DEBUG_OPUS 0
 /** Turn on entropy data in packets */
 #define ADD_ENTROPY_TO_PACKET 0
-/** Use UE4 memory allocation or Opus */
-#define USE_UE4_MEM_ALLOC 1
+/** Use UE memory allocation or Opus */
+#define USE_UE_MEM_ALLOC 1
 /** Maximum number of frames in a single Opus packet */
 #define MAX_OPUS_FRAMES_PER_PACKET 48
 /** Number of max frames for buffer calculation purposes */
@@ -235,7 +235,7 @@ bool FVoiceEncoderOpus::Init(int32 InSampleRate, int32 InNumChannels, EAudioEnco
 
 	const int32 Application = (EncodeHint == EAudioEncodeHint::VoiceEncode_Audio) ? OPUS_APPLICATION_AUDIO : OPUS_APPLICATION_VOIP;
 
-#if USE_UE4_MEM_ALLOC
+#if USE_UE_MEM_ALLOC
 	const int32 EncSize = opus_encoder_get_size(NumChannels);
 	Encoder = (OpusEncoder*)FMemory::Malloc(EncSize);
 	EncError = opus_encoder_init(Encoder, SampleRate, NumChannels, Application);
@@ -459,7 +459,7 @@ void FVoiceEncoderOpus::Reset()
 
 void FVoiceEncoderOpus::Destroy()
 {
-#if USE_UE4_MEM_ALLOC
+#if USE_UE_MEM_ALLOC
 	FMemory::Free(Encoder);
 #else
 	opus_encoder_destroy(Encoder);
@@ -523,7 +523,7 @@ bool FVoiceDecoderOpus::Init(int32 InSampleRate, int32 InNumChannels)
 
 	int32 DecError = 0;
 
-#if USE_UE4_MEM_ALLOC
+#if USE_UE_MEM_ALLOC
 	const int32 DecSize = opus_decoder_get_size(NumChannels);
 	Decoder = (OpusDecoder*)FMemory::Malloc(DecSize);
 	DecError = opus_decoder_init(Decoder, SampleRate, NumChannels);
@@ -564,7 +564,7 @@ void FVoiceDecoderOpus::Reset()
 
 void FVoiceDecoderOpus::Destroy()
 {
-#if USE_UE4_MEM_ALLOC
+#if USE_UE_MEM_ALLOC
 	FMemory::Free(Decoder);
 #else
 	opus_decoder_destroy(Decoder);
