Enabling Opus audio for SoundWaves

#rnx

diff --git a/Engine/Source/Developer/AudioFormatOpus/Private/AudioFormatOpus.cpp b/Engine/Source/Developer/AudioFormatOpus/Private/AudioFormatOpus.cpp
--- a/Engine/Source/Developer/AudioFormatOpus/Private/AudioFormatOpus.cpp
+++ b/Engine/Source/Developer/AudioFormatOpus/Private/AudioFormatOpus.cpp
@@ -7,7 +7,7 @@
 #include "Interfaces/IAudioFormat.h"
 #include "Interfaces/IAudioFormatModule.h"
 #include "OpusAudioInfo.h"
-#include "VorbisAudioInfo.h"
+#include "VorbisAudioInfo.h"	// for VorbisChannelInfo
 
 // Need to define this so that resampler.h compiles - probably a way around this somehow
 #define OUTSIDE_SPEEX
@@ -31,40 +31,42 @@ class FAudioFormatOpus : public IAudioFormat
 	enum
 	{
 		/** Version for OPUS format, this becomes part of the DDC key. */
-		UE_AUDIO_OPUS_VER = 8,
+		UE_AUDIO_OPUS_VER = 9,
 	};
 
 public:
-	virtual bool AllowParallelBuild() const override
+	bool AllowParallelBuild() const override
 	{
-		return false;
+		return true;
 	}
 
-	virtual uint16 GetVersion(FName Format) const override
+	uint16 GetVersion(FName Format) const override
 	{
 		check(Format == NAME_OPUS);
 		return UE_AUDIO_OPUS_VER;
 	}
 
 
-	virtual void GetSupportedFormats(TArray<FName>& OutFormats) const override
+	void GetSupportedFormats(TArray<FName>& OutFormats) const override
 	{
 		OutFormats.Add(NAME_OPUS);
 	}
 
-	virtual bool Cook(FName Format, const TArray<uint8>& SrcBuffer, FSoundQualityInfo& QualityInfo, TArray<uint8>& CompressedDataStore) const override
+	bool Cook(FName Format, const TArray<uint8>& SrcBuffer, FSoundQualityInfo& QualityInfo, TArray<uint8>& CompressedDataStore) const override
 	{
 		TRACE_CPUPROFILER_EVENT_SCOPE(FAudioFormatOpus::Cook);
 		check(Format == NAME_OPUS);
 
-		// Get best compatible sample rate
-		const uint16 kOpusSampleRate = GetBestOutputSampleRate(QualityInfo.SampleRate);
-		// Frame size must be one of 2.5, 5, 10, 20, 40 or 60 ms
-		const int32 kOpusFrameSizeMs = 60;
+		// For audio encoding purposes we want Full Band encoding with a 20ms frame size.
+		const uint16 kOpusSampleRate = 48000;
+		const int32 kOpusFrameSizeMs = 20;
 		// Calculate frame size required by Opus
 		const int32 kOpusFrameSizeSamples = (kOpusSampleRate * kOpusFrameSizeMs) / 1000;
 		const uint32 kSampleStride = SAMPLE_SIZE * QualityInfo.NumChannels;
 		const int32 kBytesPerFrame = kOpusFrameSizeSamples * kSampleStride;
+		// Number of silent samples to prepend that get removed after decoding.
+		const int32 kPrerollSkipCount = 3840;
+		int32 NumPaddingSamplesAtEnd = 0;
 
 		// Check whether source has compatible sample rate
 		TArray<uint8> SrcBufferCopy;
@@ -80,6 +82,9 @@ public:
 			// Take a copy of the source regardless
 			SrcBufferCopy = SrcBuffer;
 		}
+		
+		// Prepend the initial silence.
+		SrcBufferCopy.InsertZeroed(0, kPrerollSkipCount * SAMPLE_SIZE * QualityInfo.NumChannels);
 
 		// Initialise the Opus encoder
 		OpusEncoder* Encoder = NULL;
@@ -107,7 +112,7 @@ public:
 
 		// Calc frame and sample count
 		int32 FramesToEncode = SrcBufferCopy.Num() / kBytesPerFrame;
-		uint32 TrueSampleCount = SrcBufferCopy.Num() / kSampleStride;
+		uint32 TrueSampleCount = SrcBufferCopy.Num() / kSampleStride - kPrerollSkipCount;
 
 		// Pad the end of data with zeroes if it isn't exactly the size of a frame.
 		if (SrcBufferCopy.Num() % kBytesPerFrame != 0)
@@ -115,11 +120,19 @@ public:
 			int32 FrameDiff = kBytesPerFrame - (SrcBufferCopy.Num() % kBytesPerFrame);
 			SrcBufferCopy.AddZeroed(FrameDiff);
 			FramesToEncode++;
+			NumPaddingSamplesAtEnd = FrameDiff / (SAMPLE_SIZE * QualityInfo.NumChannels);
 		}
 
 		check(QualityInfo.NumChannels <= MAX_uint8);
 		check(FramesToEncode <= MAX_uint16);
-		SerializeHeaderData(CompressedData, kOpusSampleRate, TrueSampleCount, QualityInfo.NumChannels, FramesToEncode);
+		FOpusAudioInfo::FHeader Hdr;
+		Hdr.NumChannels = QualityInfo.NumChannels;
+		Hdr.SampleRate = kOpusSampleRate;
+		Hdr.ActiveSampleCount = TrueSampleCount;
+		Hdr.NumEncodedFrames = FramesToEncode;
+		Hdr.NumSilentSamplesAtBeginning = kPrerollSkipCount;
+		Hdr.NumSilentSamplesAtEnd = NumPaddingSamplesAtEnd;
+		SerializeHeaderData(CompressedData, Hdr);
 
 		// Temporary storage with more than enough to store any compressed frame
 		TArray<uint8> TempCompressedData;
@@ -153,25 +166,26 @@ public:
 		return CompressedDataStore.Num() > 0;
 	}
 
-	virtual bool CookSurround(FName Format, const TArray<TArray<uint8> >& SrcBuffers, FSoundQualityInfo& QualityInfo, TArray<uint8>& CompressedDataStore) const override
+	bool CookSurround(FName Format, const TArray<TArray<uint8> >& SrcBuffers, FSoundQualityInfo& QualityInfo, TArray<uint8>& CompressedDataStore) const override
 	{
 		TRACE_CPUPROFILER_EVENT_SCOPE(FAudioFormatOpus::CookSurround);
 		check(Format == NAME_OPUS);
 
-		// Get best compatible sample rate
-		const uint16 kOpusSampleRate = GetBestOutputSampleRate(QualityInfo.SampleRate);
-		// Frame size must be one of 2.5, 5, 10, 20, 40 or 60 ms
-		const int32 kOpusFrameSizeMs = 60;
+		// For audio encoding purposes we want Full Band encoding with a 20ms frame size.
+		const uint16 kOpusSampleRate = 48000;
+		const int32 kOpusFrameSizeMs = 20;
 		// Calculate frame size required by Opus
 		const int32 kOpusFrameSizeSamples = (kOpusSampleRate * kOpusFrameSizeMs) / 1000;
 		const uint32 kSampleStride = SAMPLE_SIZE * QualityInfo.NumChannels;
 		const int32 kBytesPerFrame = kOpusFrameSizeSamples * kSampleStride;
+		// Number of silent samples to prepend that get removed after decoding.
+		const int32 kPrerollSkipCount = 3840;
 
 		// Check whether source has compatible sample rate
 		TArray<TArray<uint8>> SrcBufferCopies;
 		if (QualityInfo.SampleRate != kOpusSampleRate)
 		{
-			for (int32 Index = 0; Index < SrcBuffers.Num(); Index++)
+			for(int32 Index=0; Index<SrcBuffers.Num(); ++Index)
 			{
 				TArray<uint8>& NewCopy = *new (SrcBufferCopies) TArray<uint8>;
 				if (!ResamplePCM(1, SrcBuffers[Index], QualityInfo.SampleRate, NewCopy, kOpusSampleRate))
@@ -186,7 +200,7 @@ public:
 			SrcBufferCopies.AddDefaulted(SrcBuffers.Num());
 
 			// Take a copy of the source regardless
-			for (int32 Index = 0; Index < SrcBuffers.Num(); Index++)
+			for(int32 Index=0; Index<SrcBuffers.Num(); ++Index)
 			{
 				SrcBufferCopies[Index] = SrcBuffers[Index];
 			}
@@ -194,7 +208,7 @@ public:
 
 		// Ensure that all channels are the same length
 		int32 SourceSize = -1;
-		for (int32 Index = 0; Index < SrcBufferCopies.Num(); Index++)
+		for (int32 Index=0; Index<SrcBufferCopies.Num(); ++Index)
 		{
 			if (!Index)
 			{
@@ -213,12 +227,18 @@ public:
 			return false;
 		}
 
+		// Prepend the initial silence.
+		for (int32 Index = 0; Index < SrcBuffers.Num(); Index++)
+		{
+			SrcBufferCopies[Index].InsertZeroed(0, kPrerollSkipCount * SAMPLE_SIZE);
+		}
+
 		// Initialise the Opus multistream encoder
 		OpusMSEncoder* Encoder = NULL;
 		int32 EncError = 0;
 		int32 streams = 0;
 		int32 coupled_streams = 0;
-		// mapping_family not documented but figured out: 0 = 1 or 2 channels, 1 = 1 to 8 channel surround sound, 255 = up to 255 channels with no surround processing
+		// Mapping_family 1 as per https://www.rfc-editor.org/rfc/rfc7845.html#section-5.1.1.2
 		int32 mapping_family = 1;
 		TArray<uint8> mapping;
 		mapping.AddUninitialized(QualityInfo.NumChannels);
@@ -245,17 +265,27 @@ public:
 
 		// Calc frame and sample count
 		int32 FramesToEncode = SourceSize / (kOpusFrameSizeSamples * SAMPLE_SIZE);
-		uint32 TrueSampleCount = SourceSize / SAMPLE_SIZE;
+		uint32 TrueSampleCount = SourceSize / SAMPLE_SIZE - kPrerollSkipCount;
 
 		// Add another frame if Source does not divide into an equal number of frames
-		if (SourceSize % (kOpusFrameSizeSamples * SAMPLE_SIZE) != 0)
+		int32 NumSamplesInLastBlock = (SourceSize / SAMPLE_SIZE) % kOpusFrameSizeSamples;
+		if (NumSamplesInLastBlock != 0)
 		{
-			FramesToEncode++;
+			// The silence to pad the last block with is handled in the compression loop below.
+			++FramesToEncode;
 		}
 
 		check(QualityInfo.NumChannels <= MAX_uint8);
 		check(FramesToEncode <= MAX_uint16);
-		SerializeHeaderData(CompressedData, kOpusSampleRate, TrueSampleCount, QualityInfo.NumChannels, FramesToEncode);
+
+		FOpusAudioInfo::FHeader Hdr;
+		Hdr.NumChannels = QualityInfo.NumChannels;
+		Hdr.SampleRate = kOpusSampleRate;
+		Hdr.ActiveSampleCount = TrueSampleCount;
+		Hdr.NumEncodedFrames = FramesToEncode;
+		Hdr.NumSilentSamplesAtBeginning = kPrerollSkipCount;
+		Hdr.NumSilentSamplesAtEnd = NumSamplesInLastBlock ? kOpusFrameSizeSamples - NumSamplesInLastBlock : 0;
+		SerializeHeaderData(CompressedData, Hdr);
 
 		// Temporary storage for source data in an interleaved format
 		TArray<uint8> TempInterleavedSrc;
@@ -320,7 +350,7 @@ public:
 		return CompressedDataStore.Num() > 0;
 	}
 
-	virtual int32 Recompress(FName Format, const TArray<uint8>& SrcBuffer, FSoundQualityInfo& QualityInfo, TArray<uint8>& OutBuffer) const override
+	int32 Recompress(FName Format, const TArray<uint8>& SrcBuffer, FSoundQualityInfo& QualityInfo, TArray<uint8>& OutBuffer) const override
 	{
 		check(Format == NAME_OPUS);
 		FOpusAudioInfo	AudioInfo;
@@ -351,18 +381,19 @@ public:
 		return CompressedDataStore.Num();
 	}
 
-	virtual int32 GetMinimumSizeForInitialChunk(FName Format, const TArray<uint8>& SrcBuffer) const override
+	int32 GetMinimumSizeForInitialChunk(FName Format, const TArray<uint8>& SrcBuffer) const override
 	{
-		// Since UE uses it's own version of the header, we hardcode the size of that here. See SplitDataForStreaming below to see our initial info.
-		return FCStringAnsi::Strlen(OPUS_ID_STRING) + 1 // Format identifier
-			+ sizeof(uint16) // Sample Rate
-			+ sizeof(uint32) // True Sample Count
-			+ sizeof(uint8)  // Number of Channels
-			+ sizeof(uint16); // Serialized Frames
+		return FOpusAudioInfo::FHeader::HeaderSize();
 	}
 
-	virtual bool SplitDataForStreaming(const TArray<uint8>& SrcBuffer, TArray<TArray<uint8>>& OutBuffers, const int32 MaxInitialChunkSize, const int32 MaxChunkSize) const override
+	bool SplitDataForStreaming(const TArray<uint8>& SrcBuffer, TArray<TArray<uint8>>& OutBuffers, const int32 MaxInitialChunkSize, const int32 MaxChunkSize) const override
 	{
+		// This should not be called if we require a streaming seek-table. 
+		if (!ensure(!RequiresStreamingSeekTable()))
+		{
+			return false;
+		}
+
 		if (SrcBuffer.Num() == 0)
 		{
 			return false;
@@ -370,30 +401,21 @@ public:
 
 		uint32 ReadOffset = 0;
 		uint32 WriteOffset = 0;
-		uint16 ProcessedFrames = 0;
+		uint32 ProcessedFrames = 0;
 		const uint8* LockedSrc = SrcBuffer.GetData();
 
-		// Read Identifier, True Sample Count, Number of channels and Frames to Encode first
-		if (FCStringAnsi::Strcmp((char*)LockedSrc, OPUS_ID_STRING) != 0)
+		FOpusAudioInfo::FHeader Hdr;
+		if (!FOpusAudioInfo::ParseHeader(Hdr, ReadOffset, LockedSrc, SrcBuffer.Num()))
 		{
 			return false;
 		}
-		ReadOffset += FCStringAnsi::Strlen(OPUS_ID_STRING) + 1;
-		uint16 SampleRate = *((uint16*)(LockedSrc + ReadOffset));
-		ReadOffset += sizeof(uint16);
-		uint32 TrueSampleCount = *((uint32*)(LockedSrc + ReadOffset));
-		ReadOffset += sizeof(uint32);
-		uint8 NumChannels = *(LockedSrc + ReadOffset);
-		ReadOffset += sizeof(uint8);
-		uint16 SerializedFrames = *((uint16*)(LockedSrc + ReadOffset));
-		ReadOffset += sizeof(uint16);
 
 		// Should always be able to store basic info in a single chunk
 		check(ReadOffset - WriteOffset <= (uint32)MaxInitialChunkSize);
 
 		int32 ChunkSize = MaxInitialChunkSize;
 
-		while (ProcessedFrames < SerializedFrames)
+		while (ProcessedFrames < Hdr.NumEncodedFrames)
 		{
 			uint16 FrameSize = *((uint16*)(LockedSrc + ReadOffset));
 
@@ -415,32 +437,47 @@ public:
 		return true;
 	}
 
-	/**
-	 * Calculate the best sample rate for the output opus data
-	 */
-	static uint16 GetBestOutputSampleRate(int32 SampleRate)
+	bool RequiresStreamingSeekTable() const override
+	{
+		return true;
+	}
+
+	bool ExtractSeekTableForStreaming(TArray<uint8>& InOutBuffer, IAudioFormat::FSeekTable& OutSeektable) const override
 	{
-		static const uint16 ValidSampleRates[] = 
+		// This should only be called if we require a streaming seek-table. 
+		if (!ensure(RequiresStreamingSeekTable()))
 		{
-			0, // not really valid, but simplifies logic below
-			8000,
-			12000,
-			16000,
-			24000,
-			48000,
-		};
-
-		// look for the next highest valid rate
-		for (int32 Index = UE_ARRAY_COUNT(ValidSampleRates) - 2; Index >= 0; Index--)
+			return false;
+		}
+
+		FOpusAudioInfo::FHeader Hdr;
+		uint32 CurrentOffset = 0;
+		uint8* Data = InOutBuffer.GetData();
+		if (!FOpusAudioInfo::ParseHeader(Hdr, CurrentOffset, Data, InOutBuffer.Num()))
 		{
-			if (SampleRate > ValidSampleRates[Index])
-			{
-				return ValidSampleRates[Index + 1];
-			}
+			return false;
 		}
-		// this should never get here!
-		check(0);
-		return 0;
+		Data += CurrentOffset;
+		int32 NumFrames = Hdr.NumEncodedFrames;
+		OutSeektable.Offsets.SetNum(NumFrames);
+		OutSeektable.Times.SetNum(NumFrames);
+
+		const uint16 kOpusSampleRate = 48000;
+		const int32 kOpusFrameSizeMs = 20;
+		const int32 kOpusFrameSizeSamples = (kOpusSampleRate * kOpusFrameSizeMs) / 1000;
+
+		uint64 SamplePos = 0;
+		uint32 CompressedPos = CurrentOffset;
+		for(int32 i=0; i<NumFrames; ++i)
+		{
+			OutSeektable.Times[i] = SamplePos;
+			OutSeektable.Offsets[i] = CompressedPos;
+			uint16 FrameSize = *reinterpret_cast<const uint16*>(Data);
+			Data += FrameSize + sizeof(uint16);
+			CompressedPos += FrameSize + sizeof(uint16);
+			SamplePos += kOpusFrameSizeSamples;
+		}
+		return true;
 	}
 
 	bool ResamplePCM(uint32 NumChannels, const TArray<uint8>& InBuffer, uint32 InSampleRate, TArray<uint8>& OutBuffer, uint32 OutSampleRate) const
@@ -491,21 +528,26 @@ public:
 
 	int32 GetBitRateFromQuality(FSoundQualityInfo& QualityInfo) const
 	{
-		// There is no perfect way to map Vorbis' Quality setting to an Opus bitrate but this 
-		// will use it as a multiplier to decide how much smaller than the original the
-		// compressed data should be
-		int32 OriginalBitRate = QualityInfo.SampleRate * QualityInfo.NumChannels * SAMPLE_SIZE * 8;
-		return (float)OriginalBitRate * FMath::GetMappedRangeValueClamped(FVector2f(1, 100), FVector2f(0.04, 0.25), (float)QualityInfo.Quality);
+		const int32 kMinBpsPerChannel = 16000;
+		const int32 kMaxBpsPerChannel = 96000;
+		const int32 kQuality = QualityInfo.Quality < 1 ? 1 : QualityInfo.Quality > 100 ? 100 : QualityInfo.Quality;
+		int32 Bps = (int32) FMath::GetMappedRangeValueClamped(FVector2f(1, 100), FVector2f(kMinBpsPerChannel, kMaxBpsPerChannel), (float)kQuality);
+		Bps *= QualityInfo.NumChannels;
+		return Bps;
 	}
 
-	void SerializeHeaderData(FMemoryWriter& CompressedData, uint16 SampleRate, uint32 TrueSampleCount, uint8 NumChannels, uint16 NumFrames) const
+	void SerializeHeaderData(FMemoryWriter& CompressedData, FOpusAudioInfo::FHeader& InHeader) const
 	{
-		const char* OpusIdentifier = OPUS_ID_STRING;
-		CompressedData.Serialize((void*)OpusIdentifier, FCStringAnsi::Strlen(OpusIdentifier) + 1);
-		CompressedData.Serialize(&SampleRate, sizeof(uint16));
-		CompressedData.Serialize(&TrueSampleCount, sizeof(uint32));
-		CompressedData.Serialize(&NumChannels, sizeof(uint8));
-		CompressedData.Serialize(&NumFrames, sizeof(uint16));
+		InHeader.Version = 0;
+		FMemory::Memcpy(InHeader.Identifier, FOpusAudioInfo::FHeader::OPUS_ID, 8);
+		CompressedData.Serialize(InHeader.Identifier, 8);
+		CompressedData.Serialize(&InHeader.Version, sizeof(uint8));
+		CompressedData.Serialize(&InHeader.NumChannels, sizeof(uint8));
+		CompressedData.Serialize(&InHeader.SampleRate, sizeof(uint16));
+		CompressedData.Serialize(&InHeader.ActiveSampleCount, sizeof(uint64));
+		CompressedData.Serialize(&InHeader.NumEncodedFrames, sizeof(uint32));
+		CompressedData.Serialize(&InHeader.NumSilentSamplesAtBeginning, sizeof(int32));
+		CompressedData.Serialize(&InHeader.NumSilentSamplesAtEnd, sizeof(int32));
 	}
 
 	void SerialiseFrameData(FMemoryWriter& CompressedData, uint8* FrameData, uint16 FrameSize) const
diff --git a/Engine/Source/Runtime/Engine/Classes/Sound/AudioSettings.h b/Engine/Source/Runtime/Engine/Classes/Sound/AudioSettings.h
--- a/Engine/Source/Runtime/Engine/Classes/Sound/AudioSettings.h
+++ b/Engine/Source/Runtime/Engine/Classes/Sound/AudioSettings.h
@@ -62,6 +62,9 @@ enum class EDefaultAudioCompressionType : uint8
 	// Uncompressed audio. Large memory usage (streamed chunks contain less audio per chunk) but extremely cheap to decode and supports all features. 
 	PCM,
 
+	// Opus is a highly versatile audio codec. It is primarily designed for interactive speech and music transmission over the Internet, but is also applicable to storage and streaming applications.
+	Opus,
+
 	// Encodes the asset to a platform specific format and will be different depending on the platform. It does not currently support seeking.
 	PlatformSpecific,
 };
diff --git a/Engine/Source/Runtime/Engine/Classes/Sound/SoundWave.h b/Engine/Source/Runtime/Engine/Classes/Sound/SoundWave.h
--- a/Engine/Source/Runtime/Engine/Classes/Sound/SoundWave.h
+++ b/Engine/Source/Runtime/Engine/Classes/Sound/SoundWave.h
@@ -323,6 +323,9 @@ enum class ESoundAssetCompressionType : uint8
 	// Uncompressed audio. Large memory usage (streamed chunks contain less audio per chunk) but extremely cheap to decode and supports all features. 
 	PCM,
 
+	// Opus is a highly versatile audio codec. It is primarily designed for interactive speech and music transmission over the Internet, but is also applicable to storage and streaming applications.
+	Opus,
+
 	// Encodes the asset to a platform specific format and will be different depending on the platform. It does not currently support seeking.
 	PlatformSpecific,
 
@@ -340,6 +343,7 @@ namespace Audio
 		case ESoundAssetCompressionType::BinkAudio:				return NAME_BINKA;
 		case ESoundAssetCompressionType::ADPCM:					return NAME_ADPCM;
 		case ESoundAssetCompressionType::PCM:					return NAME_PCM;
+		case ESoundAssetCompressionType::Opus:					return NAME_OPUS;
 		case ESoundAssetCompressionType::PlatformSpecific:		return NAME_PLATFORM_SPECIFIC;
 		case ESoundAssetCompressionType::ProjectDefined:		return NAME_PROJECT_DEFINED;
 		default:
@@ -355,6 +359,7 @@ namespace Audio
 			case EDefaultAudioCompressionType::BinkAudio:			return ESoundAssetCompressionType::BinkAudio;
 			case EDefaultAudioCompressionType::ADPCM:				return ESoundAssetCompressionType::ADPCM;
 			case EDefaultAudioCompressionType::PCM:					return ESoundAssetCompressionType::PCM;
+			case EDefaultAudioCompressionType::Opus:				return ESoundAssetCompressionType::Opus;
 			case EDefaultAudioCompressionType::PlatformSpecific:	return ESoundAssetCompressionType::PlatformSpecific;
 			default:
 				ensure(false);
diff --git a/Engine/Source/Runtime/Engine/Private/AudioFormatSettings.cpp b/Engine/Source/Runtime/Engine/Private/AudioFormatSettings.cpp
--- a/Engine/Source/Runtime/Engine/Private/AudioFormatSettings.cpp
+++ b/Engine/Source/Runtime/Engine/Private/AudioFormatSettings.cpp
@@ -78,7 +78,7 @@ namespace Audio
 			}
 			else
 			{
-				AllWaveFormats = { NAME_BINKA, NAME_ADPCM, NAME_PCM };
+				AllWaveFormats = { NAME_BINKA, NAME_ADPCM, NAME_PCM, NAME_OPUS };
 				UE_LOG(LogAudio, Warning, TEXT("Audio:AllWaveFormats is not defined, defaulting to built in formats. (%s)"), *MakePrettyArrayToString(AllWaveFormats));
 			}
 		}
diff --git a/Engine/Source/Runtime/Engine/Private/OpusAudioInfo.cpp b/Engine/Source/Runtime/Engine/Private/OpusAudioInfo.cpp
--- a/Engine/Source/Runtime/Engine/Private/OpusAudioInfo.cpp
+++ b/Engine/Source/Runtime/Engine/Private/OpusAudioInfo.cpp
@@ -6,15 +6,9 @@
 #include <opus_defines.h>
 #include <opus_types.h>
 
-#if !defined(WITH_OPUS)
-	#define WITH_OPUS (PLATFORM_WINDOWS || PLATFORM_MAC || PLATFORM_UNIX)
-#endif
-
-#if WITH_OPUS
 THIRD_PARTY_INCLUDES_START
 #include "opus_multistream.h"
 THIRD_PARTY_INCLUDES_END
-#endif
 
 #define USE_UE_MEM_ALLOC 1
 #define OPUS_MAX_FRAME_SIZE_MS 120
@@ -48,7 +42,6 @@ struct FOpusDecoderWrapper
 {
 	FOpusDecoderWrapper(uint16 SampleRate, uint8 NumChannels)
 	{
-#if WITH_OPUS
 		check(NumChannels <= 8);
 		const UnrealChannelLayout& Layout = UnrealMappings[NumChannels-1];
 	#if USE_UE_MEM_ALLOC
@@ -58,43 +51,30 @@ struct FOpusDecoderWrapper
 	#else
 		Decoder = opus_multistream_decoder_create(SampleRate, NumChannels, Layout.NumStreams, Layout.NumCoupledStreams, Layout.Mapping, &DecError);
 	#endif
-#endif
 	}
 
 	~FOpusDecoderWrapper()
 	{
-#if WITH_OPUS
 	#if USE_UE_MEM_ALLOC
 		FMemory::Free(Decoder);
 	#else
 		opus_multistream_encoder_destroy(Decoder);
 	#endif
-#endif
 	}
 
 	int32 Decode(const uint8* FrameData, uint16 FrameSize, int16* OutPCMData, int32 SampleSize)
 	{
-#if WITH_OPUS
 		return opus_multistream_decode(Decoder, FrameData, FrameSize, OutPCMData, SampleSize, 0);
-#else
-		return -1;
-#endif
 	}
 
 	bool WasInitialisedSuccessfully() const
 	{
-#if WITH_OPUS
 		return DecError == OPUS_OK;
-#else
-		return false;
-#endif
 	}
 
 private:
-#if WITH_OPUS
 	OpusMSDecoder* Decoder;
 	int32 DecError;
-#endif
 };
 
 /*------------------------------------------------------------------------------------
@@ -114,6 +94,38 @@ FOpusAudioInfo::~FOpusAudioInfo()
 	}
 }
 
+
+bool FOpusAudioInfo::ParseHeader(FHeader& OutHeader, uint32& OutNumRead, const uint8* InSrcBufferData, uint32 InSrcBufferDataSize)
+{
+	OutNumRead = 0;
+
+	if ((int32)InSrcBufferDataSize < FHeader::HeaderSize())
+	{
+		return false;
+	}
+
+	auto Read = [&InSrcBufferData, &OutNumRead](void* To, int32 NumBytes) -> void
+	{
+		FMemory::Memcpy(To, InSrcBufferData, NumBytes);
+		InSrcBufferData += NumBytes;
+		OutNumRead += NumBytes;
+	};
+
+	Read(OutHeader.Identifier, 8);
+	if (FMemory::Memcmp(OutHeader.Identifier, FHeader::OPUS_ID, 8))
+	{
+		return false;
+	}
+	Read(&OutHeader.Version, sizeof(uint8));
+	Read(&OutHeader.NumChannels, sizeof(uint8));
+	Read(&OutHeader.SampleRate, sizeof(uint16));
+	Read(&OutHeader.ActiveSampleCount, sizeof(uint64));
+	Read(&OutHeader.NumEncodedFrames, sizeof(uint32));
+	Read(&OutHeader.NumSilentSamplesAtBeginning, sizeof(int32));
+	Read(&OutHeader.NumSilentSamplesAtEnd, sizeof(int32));
+	return true;
+}
+
 bool FOpusAudioInfo::ParseHeader(const uint8* InSrcBufferData, uint32 InSrcBufferDataSize, struct FSoundQualityInfo* QualityInfo)
 {
 	SrcBufferData = InSrcBufferData;
@@ -121,35 +133,25 @@ bool FOpusAudioInfo::ParseHeader(const uint8* InSrcBufferData, uint32 InSrcBuffe
 	SrcBufferOffset = 0;
 	CurrentSampleCount = 0;
 
-	check(InSrcBufferDataSize >= sizeof(OPUS_ID_STRING));
-
-	// Read Identifier, True Sample Count, Number of channels and Frames to Encode first
-	if (!InSrcBufferData || FMemory::Memcmp((char*)InSrcBufferData, OPUS_ID_STRING, sizeof(OPUS_ID_STRING) - 1) != 0)
+	Header.Reset();
+	if (!ParseHeader(Header, SrcBufferOffset, InSrcBufferData, InSrcBufferDataSize))
 	{
 		return false;
 	}
-	else
-	{
-		SrcBufferOffset += FCStringAnsi::Strlen(OPUS_ID_STRING) + 1;
-	}
-
-	Read(&SampleRate, sizeof(uint16));
-	Read(&TrueSampleCount, sizeof(uint32));
-	Read(&NumChannels, sizeof(uint8));
-
-	uint16 SerializedFrames = 0;
-	Read(&SerializedFrames, sizeof(uint16));
 
 	// Store the offset to where the audio data begins
 	AudioDataOffset = SrcBufferOffset;
+	// Set members from header
+	TrueSampleCount = Header.ActiveSampleCount;
+	NumChannels = Header.NumChannels;
 
 	// Write out the the header info
 	if (QualityInfo)
 	{
-		QualityInfo->SampleRate = SampleRate;
-		QualityInfo->NumChannels = NumChannels;
-		QualityInfo->SampleDataSize = TrueSampleCount * QualityInfo->NumChannels * sizeof(int16);
-		QualityInfo->Duration = (float)TrueSampleCount / QualityInfo->SampleRate;
+		QualityInfo->SampleRate = Header.SampleRate;
+		QualityInfo->NumChannels = Header.NumChannels;
+		QualityInfo->SampleDataSize = (uint32)Header.ActiveSampleCount * QualityInfo->NumChannels * sizeof(int16);
+		QualityInfo->Duration = (float)Header.ActiveSampleCount / QualityInfo->SampleRate;
 	}
 
 	return true;
@@ -158,7 +160,7 @@ bool FOpusAudioInfo::ParseHeader(const uint8* InSrcBufferData, uint32 InSrcBuffe
 bool FOpusAudioInfo::CreateDecoder()
 {
 	check(OpusDecoderWrapper == nullptr);
-	OpusDecoderWrapper = new FOpusDecoderWrapper(SampleRate, NumChannels);
+	OpusDecoderWrapper = new FOpusDecoderWrapper(Header.SampleRate, NumChannels);
 	if (!OpusDecoderWrapper->WasInitialisedSuccessfully())
 	{
 		delete OpusDecoderWrapper;
@@ -166,6 +168,7 @@ bool FOpusAudioInfo::CreateDecoder()
 		return false;
 	}
 
+	NumRemainingSamplesToSkip = Header.NumSilentSamplesAtBeginning;
 	return true;
 }
 
@@ -180,7 +183,7 @@ int32 FOpusAudioInfo::GetFrameSize()
 
 uint32 FOpusAudioInfo::GetMaxFrameSizeSamples() const
 {
-	return SampleRate * OPUS_MAX_FRAME_SIZE_MS / 1000;
+	return Header.SampleRate * OPUS_MAX_FRAME_SIZE_MS / 1000;
 }
 
 FDecodeResult FOpusAudioInfo::Decode(const uint8* CompressedData, const int32 CompressedDataSize, uint8* OutPCMData, const int32 OutputPCMDataSize)
@@ -190,9 +193,38 @@ FDecodeResult FOpusAudioInfo::Decode(const uint8* CompressedData, const int32 Co
 	if (OpusDecoderWrapper)
 	{
 		const int32 SampleSize = OutputPCMDataSize / NumChannels * sizeof(int16);
+		Result.NumCompressedBytesConsumed = CompressedDataSize;
 		Result.NumAudioFramesProduced = OpusDecoderWrapper->Decode(CompressedData, CompressedDataSize, (int16*)OutPCMData, SampleSize);
+		if (NumRemainingSamplesToSkip)
+		{
+			if (NumRemainingSamplesToSkip >= Result.NumAudioFramesProduced)
+			{
+				NumRemainingSamplesToSkip -= Result.NumAudioFramesProduced;
+				Result.NumAudioFramesProduced = 0;
+			}
+			else
+			{
+				uint8* FirstUsable = OutPCMData + NumRemainingSamplesToSkip * NumChannels * sizeof(int16);
+				int32 UsableSize = (Result.NumAudioFramesProduced - NumRemainingSamplesToSkip) * NumChannels * sizeof(int16);
+				FMemory::Memmove(OutPCMData, FirstUsable, UsableSize);
+				Result.NumAudioFramesProduced -= NumRemainingSamplesToSkip;
+				NumRemainingSamplesToSkip = 0;
+			}
+		}
+		Result.NumPcmBytesProduced = Result.NumAudioFramesProduced * NumChannels * sizeof(int16);
 	}
-	
 	return Result;
 }
 
+void FOpusAudioInfo::PrepareToLoop()
+{
+	IStreamedCompressedInfo::PrepareToLoop();
+	NumRemainingSamplesToSkip = Header.NumSilentSamplesAtBeginning;
+}
+
+void FOpusAudioInfo::SeekToTime(const float InSeekTime)
+{
+	IStreamedCompressedInfo::SeekToTime(InSeekTime);
+	NumRemainingSamplesToSkip = Header.NumSilentSamplesAtBeginning;
+}
+
diff --git a/Engine/Source/Runtime/Engine/Public/OpusAudioInfo.h b/Engine/Source/Runtime/Engine/Public/OpusAudioInfo.h
--- a/Engine/Source/Runtime/Engine/Public/OpusAudioInfo.h
+++ b/Engine/Source/Runtime/Engine/Public/OpusAudioInfo.h
@@ -9,8 +9,6 @@
 #include "CoreMinimal.h"
 #include "AudioDecompress.h"
 
-#define OPUS_ID_STRING "UE4OPUS"
-
 struct FOpusDecoderWrapper;
 struct FSoundQualityInfo;
 
@@ -24,17 +22,61 @@ public:
 	ENGINE_API virtual ~FOpusAudioInfo();
 
 	//~ Begin IStreamedCompressedInfo Interface
-	virtual bool ParseHeader(const uint8* InSrcBufferData, uint32 InSrcBufferDataSize, FSoundQualityInfo* QualityInfo) override;
-	virtual int32 GetFrameSize() override;
-	virtual uint32 GetMaxFrameSizeSamples() const override;
-	virtual bool CreateDecoder() override;
-	virtual FDecodeResult Decode(const uint8* CompressedData, const int32 CompressedDataSize, uint8* OutPCMData, const int32 OutputPCMDataSize) override;
+	bool ParseHeader(const uint8* InSrcBufferData, uint32 InSrcBufferDataSize, FSoundQualityInfo* QualityInfo) override;
+	int32 GetFrameSize() override;
+	uint32 GetMaxFrameSizeSamples() const override;
+	bool CreateDecoder() override;
+	FDecodeResult Decode(const uint8* CompressedData, const int32 CompressedDataSize, uint8* OutPCMData, const int32 OutputPCMDataSize) override;
+	void PrepareToLoop() override;
+	void SeekToTime(const float SeekTime) override;
 	//~ End IStreamedCompressedInfo Interface
 
+	struct ENGINE_API FHeader
+	{
+		char Identifier[8];
+		uint8 Version = 0;
+		uint8 NumChannels = 0;
+		uint16 SampleRate = 0;
+		uint64 ActiveSampleCount = 0;
+		uint32 NumEncodedFrames = 0;
+		int32 NumSilentSamplesAtBeginning = 0;
+		int32 NumSilentSamplesAtEnd = 0;
+
+		void Reset()
+		{
+			FMemory::Memzero(Identifier);
+			Version = 0;
+			NumChannels = 0;
+			SampleRate = 0;
+			ActiveSampleCount = 0;
+			NumEncodedFrames = 0;
+			NumSilentSamplesAtBeginning = 0;
+			NumSilentSamplesAtEnd = 0;
+		}
+
+		static inline const char OPUS_ID[8] {'U','E','O','P','U','S','\0','\0'};
+		static constexpr int32 HeaderSize()
+		{
+			return
+				  sizeof(char)*8	// Identifier
+				+ sizeof(uint8)		// Version
+				+ sizeof(uint8)		// NumChannels
+				+ sizeof(uint16)	// SampleRate
+				+ sizeof(uint64)	// ActiveSampleCount
+				+ sizeof(uint32)	// NumEncodedFrames
+				+ sizeof(int32)		// NumSilentSamplesAtBeginning
+				+ sizeof(int32);	// NumSilentSamplesAtEnd
+		}
+	};
+	static ENGINE_API bool ParseHeader(FHeader& OutHeader, uint32& OutNumRead, const uint8* InSrcBufferData, uint32 InSrcBufferDataSize);
+
 protected:
 	/** Wrapper around Opus-specific decoding state and APIs */
-	FOpusDecoderWrapper*	OpusDecoderWrapper;
+	FOpusDecoderWrapper* OpusDecoderWrapper;
 
-	/** Copied from the header. Opus only need 16 bits */
-	uint16 SampleRate;
+	/** The currently active header. */
+	FHeader Header;
+	
+	/** The number of decoded samples that need to be skipped over. */
+	int32 NumRemainingSamplesToSkip = 0;
 };
