Fixed incorrect check() on number of encoded frames.

#rnx

diff --git a/Engine/Source/Developer/AudioFormatOpus/Private/AudioFormatOpus.cpp b/Engine/Source/Developer/AudioFormatOpus/Private/AudioFormatOpus.cpp
--- a/Engine/Source/Developer/AudioFormatOpus/Private/AudioFormatOpus.cpp
+++ b/Engine/Source/Developer/AudioFormatOpus/Private/AudioFormatOpus.cpp
@@ -111,7 +111,7 @@ public:
 		int32 SrcBufferOffset = 0;
 
 		// Calc frame and sample count
-		int32 FramesToEncode = SrcBufferCopy.Num() / kBytesPerFrame;
+		int64 FramesToEncode = SrcBufferCopy.Num() / kBytesPerFrame;
 		uint32 TrueSampleCount = SrcBufferCopy.Num() / kSampleStride - kPrerollSkipCount;
 
 		// Pad the end of data with zeroes if it isn't exactly the size of a frame.
@@ -124,12 +124,12 @@ public:
 		}
 
 		check(QualityInfo.NumChannels <= MAX_uint8);
-		check(FramesToEncode <= MAX_uint16);
+		check(FramesToEncode <= MAX_uint32);
 		FOpusAudioInfo::FHeader Hdr;
 		Hdr.NumChannels = QualityInfo.NumChannels;
 		Hdr.SampleRate = kOpusSampleRate;
 		Hdr.ActiveSampleCount = TrueSampleCount;
-		Hdr.NumEncodedFrames = FramesToEncode;
+		Hdr.NumEncodedFrames = (uint32)FramesToEncode;
 		Hdr.NumSilentSamplesAtBeginning = kPrerollSkipCount;
 		Hdr.NumSilentSamplesAtEnd = NumPaddingSamplesAtEnd;
 		SerializeHeaderData(CompressedData, Hdr);
@@ -264,7 +264,7 @@ public:
 		int32 SrcBufferOffset = 0;
 
 		// Calc frame and sample count
-		int32 FramesToEncode = SourceSize / (kOpusFrameSizeSamples * SAMPLE_SIZE);
+		int64 FramesToEncode = SourceSize / (kOpusFrameSizeSamples * SAMPLE_SIZE);
 		uint32 TrueSampleCount = SourceSize / SAMPLE_SIZE - kPrerollSkipCount;
 
 		// Add another frame if Source does not divide into an equal number of frames
@@ -276,13 +276,13 @@ public:
 		}
 
 		check(QualityInfo.NumChannels <= MAX_uint8);
-		check(FramesToEncode <= MAX_uint16);
+		check(FramesToEncode <= MAX_uint32);
 
 		FOpusAudioInfo::FHeader Hdr;
 		Hdr.NumChannels = QualityInfo.NumChannels;
 		Hdr.SampleRate = kOpusSampleRate;
 		Hdr.ActiveSampleCount = TrueSampleCount;
-		Hdr.NumEncodedFrames = FramesToEncode;
+		Hdr.NumEncodedFrames = (uint32) FramesToEncode;
 		Hdr.NumSilentSamplesAtBeginning = kPrerollSkipCount;
 		Hdr.NumSilentSamplesAtEnd = NumSamplesInLastBlock ? kOpusFrameSizeSamples - NumSamplesInLastBlock : 0;
 		SerializeHeaderData(CompressedData, Hdr);
