Fixed incorrect frame calculation when cooking multichannel Opus audio

#rnx

diff --git a/Engine/Source/Developer/AudioFormatOpus/Private/AudioFormatOpus.cpp b/Engine/Source/Developer/AudioFormatOpus/Private/AudioFormatOpus.cpp
--- a/Engine/Source/Developer/AudioFormatOpus/Private/AudioFormatOpus.cpp
+++ b/Engine/Source/Developer/AudioFormatOpus/Private/AudioFormatOpus.cpp
@@ -232,6 +232,7 @@ public:
 		{
 			SrcBufferCopies[Index].InsertZeroed(0, kPrerollSkipCount * SAMPLE_SIZE);
 		}
+		SourceSize += kPrerollSkipCount * SAMPLE_SIZE;
 
 		// Initialise the Opus multistream encoder
 		OpusMSEncoder* Encoder = NULL;
