Removed resampling to Opus internal rate since playback is done at the original sample rate, resulting in a pitch shift if the sample has changed rate for encoding.

#rnx

diff --git a/Engine/Source/Developer/AudioFormatOpus/Private/AudioFormatOpus.cpp b/Engine/Source/Developer/AudioFormatOpus/Private/AudioFormatOpus.cpp
--- a/Engine/Source/Developer/AudioFormatOpus/Private/AudioFormatOpus.cpp
+++ b/Engine/Source/Developer/AudioFormatOpus/Private/AudioFormatOpus.cpp
@@ -31,7 +31,7 @@ class FAudioFormatOpus : public IAudioFormat
 	enum
 	{
 		/** Version for OPUS format, this becomes part of the DDC key. */
-		UE_AUDIO_OPUS_VER = 9,
+		UE_AUDIO_OPUS_VER = 10,
 	};
 
 public:
@@ -68,22 +68,8 @@ public:
 		const int32 kPrerollSkipCount = 3840;
 		int32 NumPaddingSamplesAtEnd = 0;
 
-		// Check whether source has compatible sample rate
-		TArray<uint8> SrcBufferCopy;
-		if (QualityInfo.SampleRate != kOpusSampleRate)
-		{
-			if (!ResamplePCM(QualityInfo.NumChannels, SrcBuffer, QualityInfo.SampleRate, SrcBufferCopy, kOpusSampleRate))
-			{
-				return false;
-			}
-		}
-		else
-		{
-			// Take a copy of the source regardless
-			SrcBufferCopy = SrcBuffer;
-		}
-		
 		// Prepend the initial silence.
+		TArray<uint8> SrcBufferCopy = SrcBuffer;
 		SrcBufferCopy.InsertZeroed(0, kPrerollSkipCount * SAMPLE_SIZE * QualityInfo.NumChannels);
 
 		// Initialise the Opus encoder
@@ -181,29 +167,11 @@ public:
 		// Number of silent samples to prepend that get removed after decoding.
 		const int32 kPrerollSkipCount = 3840;
 
-		// Check whether source has compatible sample rate
 		TArray<TArray<uint8>> SrcBufferCopies;
-		if (QualityInfo.SampleRate != kOpusSampleRate)
+		SrcBufferCopies.AddDefaulted(SrcBuffers.Num());
+		for(int32 Index=0; Index<SrcBuffers.Num(); ++Index)
 		{
-			for(int32 Index=0; Index<SrcBuffers.Num(); ++Index)
-			{
-				TArray<uint8>& NewCopy = *new (SrcBufferCopies) TArray<uint8>;
-				if (!ResamplePCM(1, SrcBuffers[Index], QualityInfo.SampleRate, NewCopy, kOpusSampleRate))
-				{
-					return false;
-				}
-			}
-		}
-		else
-		{
-			SrcBufferCopies.Reset();
-			SrcBufferCopies.AddDefaulted(SrcBuffers.Num());
-
-			// Take a copy of the source regardless
-			for(int32 Index=0; Index<SrcBuffers.Num(); ++Index)
-			{
-				SrcBufferCopies[Index] = SrcBuffers[Index];
-			}
+			SrcBufferCopies[Index] = SrcBuffers[Index];
 		}
 
 		// Ensure that all channels are the same length
