Returning the original sample rate in AudioFormatOpus instead of the Opus internal rate since we otherwise end up in an infinite re-cooking loop when the sound is to be inlined

#rnx

diff --git a/Engine/Source/Developer/AudioFormatOpus/Private/AudioFormatOpus.cpp b/Engine/Source/Developer/AudioFormatOpus/Private/AudioFormatOpus.cpp
--- a/Engine/Source/Developer/AudioFormatOpus/Private/AudioFormatOpus.cpp
+++ b/Engine/Source/Developer/AudioFormatOpus/Private/AudioFormatOpus.cpp
@@ -31,7 +31,7 @@ class FAudioFormatOpus : public IAudioFormat
 	enum
 	{
 		/** Version for OPUS format, this becomes part of the DDC key. */
-		UE_AUDIO_OPUS_VER = 10,
+		UE_AUDIO_OPUS_VER = 11,
 	};
 
 public:
@@ -58,7 +58,7 @@ public:
 		check(Format == NAME_OPUS);
 
 		// For audio encoding purposes we want Full Band encoding with a 20ms frame size.
-		const uint16 kOpusSampleRate = 48000;
+		const uint32 kOpusSampleRate = 48000;
 		const int32 kOpusFrameSizeMs = 20;
 		// Calculate frame size required by Opus
 		const int32 kOpusFrameSizeSamples = (kOpusSampleRate * kOpusFrameSizeMs) / 1000;
@@ -113,7 +113,8 @@ public:
 		check(FramesToEncode <= MAX_uint32);
 		FOpusAudioInfo::FHeader Hdr;
 		Hdr.NumChannels = QualityInfo.NumChannels;
-		Hdr.SampleRate = kOpusSampleRate;
+		Hdr.SampleRate = QualityInfo.SampleRate;
+		Hdr.EncodedSampleRate = kOpusSampleRate;
 		Hdr.ActiveSampleCount = TrueSampleCount;
 		Hdr.NumEncodedFrames = (uint32)FramesToEncode;
 		Hdr.NumSilentSamplesAtBeginning = kPrerollSkipCount;
@@ -158,7 +159,7 @@ public:
 		check(Format == NAME_OPUS);
 
 		// For audio encoding purposes we want Full Band encoding with a 20ms frame size.
-		const uint16 kOpusSampleRate = 48000;
+		const uint32 kOpusSampleRate = 48000;
 		const int32 kOpusFrameSizeMs = 20;
 		// Calculate frame size required by Opus
 		const int32 kOpusFrameSizeSamples = (kOpusSampleRate * kOpusFrameSizeMs) / 1000;
@@ -249,7 +250,8 @@ public:
 
 		FOpusAudioInfo::FHeader Hdr;
 		Hdr.NumChannels = QualityInfo.NumChannels;
-		Hdr.SampleRate = kOpusSampleRate;
+		Hdr.SampleRate = QualityInfo.SampleRate;
+		Hdr.EncodedSampleRate = kOpusSampleRate;
 		Hdr.ActiveSampleCount = TrueSampleCount;
 		Hdr.NumEncodedFrames = (uint32) FramesToEncode;
 		Hdr.NumSilentSamplesAtBeginning = kPrerollSkipCount;
@@ -512,7 +514,8 @@ public:
 		CompressedData.Serialize(InHeader.Identifier, 8);
 		CompressedData.Serialize(&InHeader.Version, sizeof(uint8));
 		CompressedData.Serialize(&InHeader.NumChannels, sizeof(uint8));
-		CompressedData.Serialize(&InHeader.SampleRate, sizeof(uint16));
+		CompressedData.Serialize(&InHeader.SampleRate, sizeof(uint32));
+		CompressedData.Serialize(&InHeader.EncodedSampleRate, sizeof(uint32));
 		CompressedData.Serialize(&InHeader.ActiveSampleCount, sizeof(uint64));
 		CompressedData.Serialize(&InHeader.NumEncodedFrames, sizeof(uint32));
 		CompressedData.Serialize(&InHeader.NumSilentSamplesAtBeginning, sizeof(int32));
diff --git a/Engine/Source/Runtime/Engine/Private/OpusAudioInfo.cpp b/Engine/Source/Runtime/Engine/Private/OpusAudioInfo.cpp
--- a/Engine/Source/Runtime/Engine/Private/OpusAudioInfo.cpp
+++ b/Engine/Source/Runtime/Engine/Private/OpusAudioInfo.cpp
@@ -42,7 +42,7 @@ FOpusDecoderWrapper
 ------------------------------------------------------------------------------------*/
 struct FOpusDecoderWrapper
 {
-	FOpusDecoderWrapper(uint16 SampleRate, uint8 NumChannels)
+	FOpusDecoderWrapper(uint32 SampleRate, uint8 NumChannels)
 	{
 		check(NumChannels <= 8);
 		const UnrealChannelLayout& Layout = UnrealMappings[NumChannels-1];
@@ -120,7 +120,8 @@ bool FOpusAudioInfo::ParseHeader(FHeader& OutHeader, uint32& OutNumRead, const u
 	}
 	Read(&OutHeader.Version, sizeof(uint8));
 	Read(&OutHeader.NumChannels, sizeof(uint8));
-	Read(&OutHeader.SampleRate, sizeof(uint16));
+	Read(&OutHeader.SampleRate, sizeof(uint32));
+	Read(&OutHeader.EncodedSampleRate, sizeof(uint32));
 	Read(&OutHeader.ActiveSampleCount, sizeof(uint64));
 	Read(&OutHeader.NumEncodedFrames, sizeof(uint32));
 	Read(&OutHeader.NumSilentSamplesAtBeginning, sizeof(int32));
@@ -162,7 +163,7 @@ bool FOpusAudioInfo::ParseHeader(const uint8* InSrcBufferData, uint32 InSrcBuffe
 bool FOpusAudioInfo::CreateDecoder()
 {
 	check(OpusDecoderWrapper == nullptr);
-	OpusDecoderWrapper = new FOpusDecoderWrapper(Header.SampleRate, NumChannels);
+	OpusDecoderWrapper = new FOpusDecoderWrapper(Header.EncodedSampleRate, NumChannels);
 	if (!OpusDecoderWrapper->WasInitialisedSuccessfully())
 	{
 		delete OpusDecoderWrapper;
@@ -186,7 +187,7 @@ int32 FOpusAudioInfo::GetFrameSize()
 
 uint32 FOpusAudioInfo::GetMaxFrameSizeSamples() const
 {
-	return Header.SampleRate * OPUS_MAX_FRAME_SIZE_MS / 1000;
+	return Header.EncodedSampleRate * OPUS_MAX_FRAME_SIZE_MS / 1000;
 }
 
 FDecodeResult FOpusAudioInfo::Decode(const uint8* CompressedData, const int32 CompressedDataSize, uint8* OutPCMData, const int32 OutputPCMDataSize)
@@ -339,7 +340,7 @@ void FOpusAudioInfo::SeekToTime(const float InSeekTime)
 		{
 			uint32 ChunkSize = (uint32)ChunkPtr[0] + ((uint32)ChunkPtr[1] << 8);
 			int32 ExpectedFrames = opus_packet_get_nb_frames(ChunkPtr + 2, ChunkSize);
-			int32 ExpectedFrameSize = opus_packet_get_samples_per_frame(ChunkPtr + 2, Header.SampleRate);
+			int32 ExpectedFrameSize = opus_packet_get_samples_per_frame(ChunkPtr + 2, Header.EncodedSampleRate);
 			int32 NumExpectedTotal = ExpectedFrames * ExpectedFrameSize;
 
 			if (CurrentChunkSampleNum >= SeekSampleNum && SeekSampleNum < CurrentChunkSampleNum + NumExpectedTotal)
diff --git a/Engine/Source/Runtime/Engine/Public/OpusAudioInfo.h b/Engine/Source/Runtime/Engine/Public/OpusAudioInfo.h
--- a/Engine/Source/Runtime/Engine/Public/OpusAudioInfo.h
+++ b/Engine/Source/Runtime/Engine/Public/OpusAudioInfo.h
@@ -36,7 +36,8 @@ public:
 		char Identifier[8];
 		uint8 Version = 0;
 		uint8 NumChannels = 0;
-		uint16 SampleRate = 0;
+		uint32 SampleRate = 0;
+		uint32 EncodedSampleRate = 0;
 		uint64 ActiveSampleCount = 0;
 		uint32 NumEncodedFrames = 0;
 		int32 NumSilentSamplesAtBeginning = 0;
@@ -48,6 +49,7 @@ public:
 			Version = 0;
 			NumChannels = 0;
 			SampleRate = 0;
+			EncodedSampleRate = 0;
 			ActiveSampleCount = 0;
 			NumEncodedFrames = 0;
 			NumSilentSamplesAtBeginning = 0;
@@ -61,7 +63,8 @@ public:
 				  sizeof(char)*8	// Identifier
 				+ sizeof(uint8)		// Version
 				+ sizeof(uint8)		// NumChannels
-				+ sizeof(uint16)	// SampleRate
+				+ sizeof(uint32)	// SampleRate
+				+ sizeof(uint32)	// EncodedSampleRate
 				+ sizeof(uint64)	// ActiveSampleCount
 				+ sizeof(uint32)	// NumEncodedFrames
 				+ sizeof(int32)		// NumSilentSamplesAtBeginning
