Using nearest Opus native sample rate when encoding a sound wave at a non-native rate; storing and applying the pre-skip value; added an empty block at encode time to get the last remaining samples out of the decoder to allow for frame-accurate looping

#rnx

diff --git a/Engine/Source/Developer/AudioFormatOpus/Private/AudioFormatOpus.cpp b/Engine/Source/Developer/AudioFormatOpus/Private/AudioFormatOpus.cpp
--- a/Engine/Source/Developer/AudioFormatOpus/Private/AudioFormatOpus.cpp
+++ b/Engine/Source/Developer/AudioFormatOpus/Private/AudioFormatOpus.cpp
@@ -9,12 +9,8 @@
 #include "OpusAudioInfo.h"
 #include "VorbisAudioInfo.h"	// for VorbisChannelInfo
 
-// Need to define this so that resampler.h compiles - probably a way around this somehow
-#define OUTSIDE_SPEEX
-
 THIRD_PARTY_INCLUDES_START
 #include "opus_multistream.h"
-#include "speex_resampler.h"
 THIRD_PARTY_INCLUDES_END
 
 /** Use UE memory allocation or Opus */
@@ -31,7 +27,7 @@ class FAudioFormatOpus : public IAudioFormat
 	enum
 	{
 		/** Version for OPUS format, this becomes part of the DDC key. */
-		UE_AUDIO_OPUS_VER = 11,
+		UE_AUDIO_OPUS_VER = 12,
 	};
 
 public:
@@ -58,16 +54,19 @@ public:
 		check(Format == NAME_OPUS);
 
 		// For audio encoding purposes we want Full Band encoding with a 20ms frame size.
-		const uint32 kOpusSampleRate = 48000;
+		const uint32 kOpusSampleRate = GetMatchingOpusSampleRate(QualityInfo.SampleRate);
 		const int32 kOpusFrameSizeMs = 20;
 		// Calculate frame size required by Opus
 		const int32 kOpusFrameSizeSamples = (kOpusSampleRate * kOpusFrameSizeMs) / 1000;
 		const uint32 kSampleStride = SAMPLE_SIZE * QualityInfo.NumChannels;
 		const int32 kBytesPerFrame = kOpusFrameSizeSamples * kSampleStride;
 		// Number of silent samples to prepend that get removed after decoding.
-		const int32 kPrerollSkipCount = 3840;
+		const int32 kPrerollSkipCount = kOpusFrameSizeSamples * (80 / kOpusFrameSizeMs);	// 80ms worth
 		int32 NumPaddingSamplesAtEnd = 0;
 
+		// Remember the actual number of samples
+		uint32 TrueSampleCount = SrcBuffer.Num() / kSampleStride;
+
 		// Prepend the initial silence.
 		TArray<uint8> SrcBufferCopy = SrcBuffer;
 		SrcBufferCopy.InsertZeroed(0, kPrerollSkipCount * SAMPLE_SIZE * QualityInfo.NumChannels);
@@ -90,6 +89,13 @@ public:
 
 		int32 BitRate = GetBitRateFromQuality(QualityInfo);
 		opus_encoder_ctl(Encoder, OPUS_SET_BITRATE(BitRate));
+		
+		// Get the number of pre-skip samples. These are to be skipped in addition to the initial silence.
+		int32 PreSkip = 0;
+		opus_encoder_ctl(Encoder, OPUS_GET_LOOKAHEAD(&PreSkip));
+		// We add silence to the end of the buffer as indicated by the pre-skip so we get this implicit
+		// decoder delay accounted for at the end. Otherwise the last samples may not be emitted by the decoder.
+		SrcBufferCopy.AddZeroed(PreSkip * SAMPLE_SIZE * QualityInfo.NumChannels);
 
 		// Create a buffer to store compressed data
 		CompressedDataStore.Empty();
@@ -98,14 +104,12 @@ public:
 
 		// Calc frame and sample count
 		int64 FramesToEncode = SrcBufferCopy.Num() / kBytesPerFrame;
-		uint32 TrueSampleCount = SrcBufferCopy.Num() / kSampleStride - kPrerollSkipCount;
-
 		// Pad the end of data with zeroes if it isn't exactly the size of a frame.
 		if (SrcBufferCopy.Num() % kBytesPerFrame != 0)
 		{
 			int32 FrameDiff = kBytesPerFrame - (SrcBufferCopy.Num() % kBytesPerFrame);
 			SrcBufferCopy.AddZeroed(FrameDiff);
-			FramesToEncode++;
+			++FramesToEncode;
 			NumPaddingSamplesAtEnd = FrameDiff / (SAMPLE_SIZE * QualityInfo.NumChannels);
 		}
 
@@ -117,6 +121,7 @@ public:
 		Hdr.EncodedSampleRate = kOpusSampleRate;
 		Hdr.ActiveSampleCount = TrueSampleCount;
 		Hdr.NumEncodedFrames = (uint32)FramesToEncode;
+		Hdr.NumPreSkipSamples = PreSkip;
 		Hdr.NumSilentSamplesAtBeginning = kPrerollSkipCount;
 		Hdr.NumSilentSamplesAtEnd = NumPaddingSamplesAtEnd;
 		SerializeHeaderData(CompressedData, Hdr);
@@ -159,14 +164,14 @@ public:
 		check(Format == NAME_OPUS);
 
 		// For audio encoding purposes we want Full Band encoding with a 20ms frame size.
-		const uint32 kOpusSampleRate = 48000;
+		const uint32 kOpusSampleRate = GetMatchingOpusSampleRate(QualityInfo.SampleRate);
 		const int32 kOpusFrameSizeMs = 20;
 		// Calculate frame size required by Opus
 		const int32 kOpusFrameSizeSamples = (kOpusSampleRate * kOpusFrameSizeMs) / 1000;
 		const uint32 kSampleStride = SAMPLE_SIZE * QualityInfo.NumChannels;
 		const int32 kBytesPerFrame = kOpusFrameSizeSamples * kSampleStride;
 		// Number of silent samples to prepend that get removed after decoding.
-		const int32 kPrerollSkipCount = 3840;
+		const int32 kPrerollSkipCount = kOpusFrameSizeSamples * (80 / kOpusFrameSizeMs);	// 80ms worth
 
 		TArray<TArray<uint8>> SrcBufferCopies;
 		SrcBufferCopies.AddDefaulted(SrcBuffers.Num());
@@ -196,6 +201,9 @@ public:
 			return false;
 		}
 
+		// Remember the actual number of samples
+		uint32 TrueSampleCount = SourceSize / SAMPLE_SIZE;
+
 		// Prepend the initial silence.
 		for (int32 Index = 0; Index < SrcBuffers.Num(); Index++)
 		{
@@ -228,6 +236,17 @@ public:
 		int32 BitRate = GetBitRateFromQuality(QualityInfo);
 		opus_multistream_encoder_ctl(Encoder, OPUS_SET_BITRATE(BitRate));
 
+		// Get the number of pre-skip samples. These are to be skipped in addition to the initial silence.
+		int32 PreSkip = 0;
+		opus_multistream_encoder_ctl(Encoder, OPUS_GET_LOOKAHEAD(&PreSkip));
+		// We add silence to the end of the buffer as indicated by the pre-skip so we get this implicit
+		// decoder delay accounted for at the end. Otherwise the last samples may not be emitted by the decoder.
+		for (int32 Index = 0; Index < SrcBuffers.Num(); Index++)
+		{
+			SrcBufferCopies[Index].AddZeroed(PreSkip * SAMPLE_SIZE);
+		}
+		SourceSize += PreSkip * SAMPLE_SIZE;
+
 		// Create a buffer to store compressed data
 		CompressedDataStore.Empty();
 		FMemoryWriter CompressedData(CompressedDataStore);
@@ -235,8 +254,6 @@ public:
 
 		// Calc frame and sample count
 		int64 FramesToEncode = SourceSize / (kOpusFrameSizeSamples * SAMPLE_SIZE);
-		uint32 TrueSampleCount = SourceSize / SAMPLE_SIZE - kPrerollSkipCount;
-
 		// Add another frame if Source does not divide into an equal number of frames
 		int32 NumSamplesInLastBlock = (SourceSize / SAMPLE_SIZE) % kOpusFrameSizeSamples;
 		if (NumSamplesInLastBlock != 0)
@@ -254,6 +271,7 @@ public:
 		Hdr.EncodedSampleRate = kOpusSampleRate;
 		Hdr.ActiveSampleCount = TrueSampleCount;
 		Hdr.NumEncodedFrames = (uint32) FramesToEncode;
+		Hdr.NumPreSkipSamples = PreSkip;
 		Hdr.NumSilentSamplesAtBeginning = kPrerollSkipCount;
 		Hdr.NumSilentSamplesAtEnd = NumSamplesInLastBlock ? kOpusFrameSizeSamples - NumSamplesInLastBlock : 0;
 		SerializeHeaderData(CompressedData, Hdr);
@@ -451,61 +469,37 @@ public:
 		return true;
 	}
 
-	bool ResamplePCM(uint32 NumChannels, const TArray<uint8>& InBuffer, uint32 InSampleRate, TArray<uint8>& OutBuffer, uint32 OutSampleRate) const
+	int32 GetBitRateFromQuality(FSoundQualityInfo& QualityInfo) const
 	{
-		// Initialize resampler to convert to desired rate for Opus
-		int32 err = 0;
-		SpeexResamplerState* resampler = speex_resampler_init(NumChannels, InSampleRate, OutSampleRate, SPEEX_RESAMPLER_QUALITY_DESKTOP, &err);
-		if (err != RESAMPLER_ERR_SUCCESS)
-		{
-			speex_resampler_destroy(resampler);
-			return false;
-		}
+		const int32 kMinBpsPerChannel = 16000;
+		const int32 kMaxBpsPerChannel = 96000;
+		const int32 kQuality = QualityInfo.Quality < 1 ? 1 : QualityInfo.Quality > 100 ? 100 : QualityInfo.Quality;
+		int32 Bps = (int32) FMath::GetMappedRangeValueClamped(FVector2f(1, 100), FVector2f(kMinBpsPerChannel, kMaxBpsPerChannel), (float)kQuality);
+		Bps *= QualityInfo.NumChannels;
+		return Bps;
+	}
 
-		// Calculate extra space required for sample rate
-		const uint32 SampleStride = SAMPLE_SIZE * NumChannels;
-		const float Duration = (float)InBuffer.Num() / (InSampleRate * SampleStride);
-		const int32 SafeCopySize = (Duration + 1) * OutSampleRate * SampleStride;
-		OutBuffer.Empty(SafeCopySize);
-		OutBuffer.AddUninitialized(SafeCopySize);
-		uint32 InSamples = InBuffer.Num() / SampleStride;
-		uint32 OutSamples = OutBuffer.Num() / SampleStride;
-
-		// Do resampling and check results
-		if (NumChannels == 1)
+	uint32 GetMatchingOpusSampleRate(uint32 InRate) const
+	{
+		if (InRate <= 8000)
 		{
-			err = speex_resampler_process_int(resampler, 0, (const short*)(InBuffer.GetData()), &InSamples, (short*)(OutBuffer.GetData()), &OutSamples);
+			return 8000;
 		}
-		else
+		else if (InRate <= 12000)
 		{
-			err = speex_resampler_process_interleaved_int(resampler, (const short*)(InBuffer.GetData()), &InSamples, (short*)(OutBuffer.GetData()), &OutSamples);
+			return 12000;
 		}
-
-		speex_resampler_destroy(resampler);
-		if (err != RESAMPLER_ERR_SUCCESS)
+		if (InRate <= 16000)
 		{
-			return false;
+			return 16000;
 		}
-
-		// reduce the size of Out Buffer if more space than necessary was allocated
-		const int32 WrittenBytes = (int32)(OutSamples * SampleStride);
-		if (WrittenBytes < OutBuffer.Num())
+		if (InRate <= 24000)
 		{
-			OutBuffer.SetNum(WrittenBytes, true);
+			return 24000;
 		}
-
-		return true;
+		return 48000;
 	}
 
-	int32 GetBitRateFromQuality(FSoundQualityInfo& QualityInfo) const
-	{
-		const int32 kMinBpsPerChannel = 16000;
-		const int32 kMaxBpsPerChannel = 96000;
-		const int32 kQuality = QualityInfo.Quality < 1 ? 1 : QualityInfo.Quality > 100 ? 100 : QualityInfo.Quality;
-		int32 Bps = (int32) FMath::GetMappedRangeValueClamped(FVector2f(1, 100), FVector2f(kMinBpsPerChannel, kMaxBpsPerChannel), (float)kQuality);
-		Bps *= QualityInfo.NumChannels;
-		return Bps;
-	}
 
 	void SerializeHeaderData(FMemoryWriter& CompressedData, FOpusAudioInfo::FHeader& InHeader) const
 	{
@@ -518,6 +512,7 @@ public:
 		CompressedData.Serialize(&InHeader.EncodedSampleRate, sizeof(uint32));
 		CompressedData.Serialize(&InHeader.ActiveSampleCount, sizeof(uint64));
 		CompressedData.Serialize(&InHeader.NumEncodedFrames, sizeof(uint32));
+		CompressedData.Serialize(&InHeader.NumPreSkipSamples, sizeof(int32));
 		CompressedData.Serialize(&InHeader.NumSilentSamplesAtBeginning, sizeof(int32));
 		CompressedData.Serialize(&InHeader.NumSilentSamplesAtEnd, sizeof(int32));
 	}
diff --git a/Engine/Source/Runtime/Engine/Private/OpusAudioInfo.cpp b/Engine/Source/Runtime/Engine/Private/OpusAudioInfo.cpp
--- a/Engine/Source/Runtime/Engine/Private/OpusAudioInfo.cpp
+++ b/Engine/Source/Runtime/Engine/Private/OpusAudioInfo.cpp
@@ -11,7 +11,6 @@ THIRD_PARTY_INCLUDES_START
 THIRD_PARTY_INCLUDES_END
 
 #define USE_UE_MEM_ALLOC 1
-#define OPUS_MAX_FRAME_SIZE_MS 120
 
 DEFINE_LOG_CATEGORY_STATIC(LogOpusAudioDecoder, Log, All);
 
@@ -60,7 +59,7 @@ struct FOpusDecoderWrapper
 	#if USE_UE_MEM_ALLOC
 		FMemory::Free(Decoder);
 	#else
-		opus_multistream_encoder_destroy(Decoder);
+		opus_multistream_decoder_destroy(Decoder);
 	#endif
 	}
 
@@ -124,6 +123,7 @@ bool FOpusAudioInfo::ParseHeader(FHeader& OutHeader, uint32& OutNumRead, const u
 	Read(&OutHeader.EncodedSampleRate, sizeof(uint32));
 	Read(&OutHeader.ActiveSampleCount, sizeof(uint64));
 	Read(&OutHeader.NumEncodedFrames, sizeof(uint32));
+	Read(&OutHeader.NumPreSkipSamples, sizeof(int32));
 	Read(&OutHeader.NumSilentSamplesAtBeginning, sizeof(int32));
 	Read(&OutHeader.NumSilentSamplesAtEnd, sizeof(int32));
 	return true;
@@ -154,7 +154,7 @@ bool FOpusAudioInfo::ParseHeader(const uint8* InSrcBufferData, uint32 InSrcBuffe
 		QualityInfo->SampleRate = Header.SampleRate;
 		QualityInfo->NumChannels = Header.NumChannels;
 		QualityInfo->SampleDataSize = (uint32)Header.ActiveSampleCount * QualityInfo->NumChannels * sizeof(int16);
-		QualityInfo->Duration = (float)Header.ActiveSampleCount / QualityInfo->SampleRate;
+		QualityInfo->Duration = (float)((double)Header.ActiveSampleCount / QualityInfo->SampleRate);
 	}
 
 	return true;
@@ -171,7 +171,7 @@ bool FOpusAudioInfo::CreateDecoder()
 		return false;
 	}
 
-	NumRemainingSamplesToSkip = Header.NumSilentSamplesAtBeginning;
+	NumRemainingSamplesToSkip = Header.NumSilentSamplesAtBeginning + Header.NumPreSkipSamples;
 	PreviousDecodedUnusedSamples.Empty();
 	return true;
 }
@@ -187,7 +187,10 @@ int32 FOpusAudioInfo::GetFrameSize()
 
 uint32 FOpusAudioInfo::GetMaxFrameSizeSamples() const
 {
-	return Header.EncodedSampleRate * OPUS_MAX_FRAME_SIZE_MS / 1000;
+	// The encoder is using 20ms frame sizes.
+	const int32 OPUS_MAX_FRAME_SIZE_MS = 20;
+	// There can be at most 2 frames in one packet, so multiply by 2.
+	return Header.EncodedSampleRate * OPUS_MAX_FRAME_SIZE_MS * 2 / 1000;
 }
 
 FDecodeResult FOpusAudioInfo::Decode(const uint8* CompressedData, const int32 CompressedDataSize, uint8* OutPCMData, const int32 OutputPCMDataSize)
@@ -267,7 +270,7 @@ FDecodeResult FOpusAudioInfo::Decode(const uint8* CompressedData, const int32 Co
 					else
 					{
 						uint8* FirstUsable = OutputDataPtr + NumRemainingSamplesToSkip * kFrameBytes;
-						int32 UsableSize = (Result.NumAudioFramesProduced - NumRemainingSamplesToSkip) * kFrameBytes;
+						int32 UsableSize = (NumDecodedFrames - NumRemainingSamplesToSkip) * kFrameBytes;
 						FMemory::Memmove(OutputDataPtr, FirstUsable, UsableSize);
 						NumDecodedFrames -= NumRemainingSamplesToSkip;
 						NumRemainingSamplesToSkip = 0;
@@ -315,7 +318,7 @@ FDecodeResult FOpusAudioInfo::Decode(const uint8* CompressedData, const int32 Co
 void FOpusAudioInfo::PrepareToLoop()
 {
 	IStreamedCompressedInfo::PrepareToLoop();
-	NumRemainingSamplesToSkip = Header.NumSilentSamplesAtBeginning;
+	NumRemainingSamplesToSkip = Header.NumSilentSamplesAtBeginning + Header.NumPreSkipSamples;
 	PreviousDecodedUnusedSamples.Empty();
 }
 
@@ -360,7 +363,7 @@ void FOpusAudioInfo::SeekToTime(const float InSeekTime)
 		}
 	}
 
-	NumRemainingSamplesToSkip = Header.NumSilentSamplesAtBeginning;
+	NumRemainingSamplesToSkip = Header.NumSilentSamplesAtBeginning + Header.NumPreSkipSamples;
 	PreviousDecodedUnusedSamples.Empty();
 }
 
diff --git a/Engine/Source/Runtime/Engine/Public/OpusAudioInfo.h b/Engine/Source/Runtime/Engine/Public/OpusAudioInfo.h
--- a/Engine/Source/Runtime/Engine/Public/OpusAudioInfo.h
+++ b/Engine/Source/Runtime/Engine/Public/OpusAudioInfo.h
@@ -40,6 +40,7 @@ public:
 		uint32 EncodedSampleRate = 0;
 		uint64 ActiveSampleCount = 0;
 		uint32 NumEncodedFrames = 0;
+		int32 NumPreSkipSamples = 0;
 		int32 NumSilentSamplesAtBeginning = 0;
 		int32 NumSilentSamplesAtEnd = 0;
 
@@ -52,6 +53,7 @@ public:
 			EncodedSampleRate = 0;
 			ActiveSampleCount = 0;
 			NumEncodedFrames = 0;
+			NumPreSkipSamples = 0;
 			NumSilentSamplesAtBeginning = 0;
 			NumSilentSamplesAtEnd = 0;
 		}
@@ -67,6 +69,7 @@ public:
 				+ sizeof(uint32)	// EncodedSampleRate
 				+ sizeof(uint64)	// ActiveSampleCount
 				+ sizeof(uint32)	// NumEncodedFrames
+				+ sizeof(int32)		// NumPreSkipSamples
 				+ sizeof(int32)		// NumSilentSamplesAtBeginning
 				+ sizeof(int32);	// NumSilentSamplesAtEnd
 		}
